Missing file from last checkin

diff --git a/Engine/Source/Editor/NiagaraEditor/Public/INiagaraEffectEditor.h b/Engine/Source/Editor/NiagaraEditor/Public/INiagaraEffectEditor.h
new file mode 100644
--- /dev/null
+++ b/Engine/Source/Editor/NiagaraEditor/Public/INiagaraEffectEditor.h
@@ -0,0 +1,19 @@
+// Copyright 1998-2014 Epic Games, Inc. All Rights Reserved.
+
+#pragma once
+
+#include "Toolkits/IToolkit.h"
+#include "Toolkits/AssetEditorToolkit.h"
+
+
+/** Niagara Editor public interface */
+class INiagaraEffectEditor : public FAssetEditorToolkit
+{
+
+public:
+
+	// ...
+
+};
+
+
