Missed some

diff --git a/Engine/Source/Editor/NiagaraEditor/Private/NiagaraCompiler.cpp b/Engine/Source/Editor/NiagaraEditor/Private/NiagaraCompiler.cpp
--- a/Engine/Source/Editor/NiagaraEditor/Private/NiagaraCompiler.cpp
+++ b/Engine/Source/Editor/NiagaraEditor/Private/NiagaraCompiler.cpp
@@ -286,3 +286,5 @@ TNiagaraExprPtr FNiagaraCompiler::GetExternalConstant(FName ConstantName, const
 	SetOrAddConstant(false, ConstantName, Default);
 	return Expression_GetExternalConstant(ConstantName, ENiagaraDataType::Matrix);
 }
+
+#undef LOCTEXT_NAMESPACE
\ No newline at end of file
diff --git a/Engine/Source/Editor/NiagaraEditor/Private/NiagaraCompiler_VectorVM.cpp b/Engine/Source/Editor/NiagaraEditor/Private/NiagaraCompiler_VectorVM.cpp
--- a/Engine/Source/Editor/NiagaraEditor/Private/NiagaraCompiler_VectorVM.cpp
+++ b/Engine/Source/Editor/NiagaraEditor/Private/NiagaraCompiler_VectorVM.cpp
@@ -27,7 +27,7 @@ public:
 		SourceExpressions = InputExpressions;
 	}
 
-	FNiagaraExpression_VMOperation::~FNiagaraExpression_VMOperation()
+	virtual ~FNiagaraExpression_VMOperation()
 	{
 		if (ResultLocation == ENiagaraExpressionResultLocation::Temporaries && ResultIndex != INDEX_NONE)
 		{
@@ -196,7 +196,7 @@ void FNiagaraCompiler_VectorVM::CompileScript(UNiagaraScript* InScript)
 		Source->UpdateGraph->GetNodesOfClass(OutputNodes);
 		if (OutputNodes.Num() != 1)
 		{
-			UE_LOG(LogNiagaraCompiler, Error, TEXT("Script contains %s output nodes: %s"),
+			UE_LOG(LogNiagaraCompiler_VectorVM, Error, TEXT("Script contains %s output nodes: %s"),
 				OutputNodes.Num() == 0 ? TEXT("no") : TEXT("too many"),
 				*InScript->GetPathName()
 				);
