Non Unity compile fix

diff --git a/Engine/Source/Editor/NiagaraEditor/Private/NiagaraCompiler.cpp b/Engine/Source/Editor/NiagaraEditor/Private/NiagaraCompiler.cpp
--- a/Engine/Source/Editor/NiagaraEditor/Private/NiagaraCompiler.cpp
+++ b/Engine/Source/Editor/NiagaraEditor/Private/NiagaraCompiler.cpp
@@ -2,7 +2,6 @@
 
 #include "NiagaraEditorPrivatePCH.h"
 #include "Engine/NiagaraScript.h"
-#include "Engine/NiagaraConstants.h"
 #include "Components/NiagaraComponent.h"
 #include "CompilerResultsLog.h"
 #include "EdGraphUtilities.h"
diff --git a/Engine/Source/Editor/NiagaraEditor/Private/NiagaraCompiler_VectorVM.cpp b/Engine/Source/Editor/NiagaraEditor/Private/NiagaraCompiler_VectorVM.cpp
--- a/Engine/Source/Editor/NiagaraEditor/Private/NiagaraCompiler_VectorVM.cpp
+++ b/Engine/Source/Editor/NiagaraEditor/Private/NiagaraCompiler_VectorVM.cpp
@@ -2,7 +2,6 @@
 
 #include "NiagaraEditorPrivatePCH.h"
 #include "Engine/NiagaraScript.h"
-#include "Engine/NiagaraConstants.h"
 #include "Components/NiagaraComponent.h"
 #include "CompilerResultsLog.h"
 #include "EdGraphUtilities.h"
diff --git a/Engine/Source/Runtime/Engine/Private/NiagaraComponent.cpp b/Engine/Source/Runtime/Engine/Private/NiagaraComponent.cpp
--- a/Engine/Source/Runtime/Engine/Private/NiagaraComponent.cpp
+++ b/Engine/Source/Runtime/Engine/Private/NiagaraComponent.cpp
@@ -6,7 +6,6 @@
 #include "VectorVM.h"
 #include "ParticleHelper.h"
 #include "Particles/ParticleResources.h"
-#include "Engine/NiagaraConstants.h"
 #include "Engine/NiagaraEffectRenderer.h"
 #include "MeshBatch.h"
 #include "SceneUtils.h"
