Fixes to Niagara for non-unity

diff --git a/Engine/Source/Editor/NiagaraEditor/Private/SNiagaraEffectEditorViewport.cpp b/Engine/Source/Editor/NiagaraEditor/Private/SNiagaraEffectEditorViewport.cpp
--- a/Engine/Source/Editor/NiagaraEditor/Private/SNiagaraEffectEditorViewport.cpp
+++ b/Engine/Source/Editor/NiagaraEditor/Private/SNiagaraEffectEditorViewport.cpp
@@ -8,6 +8,9 @@
 #include "Runtime/Engine/Public/Slate/SceneViewport.h"
 #include "NiagaraEditor.h"
 #include "ComponentReregisterContext.h"
+#include "Components/NiagaraComponent.h"
+#include "Engine/NiagaraEffect.h"
+#include "SDockTab.h"
 
 /** Viewport Client for the preview viewport */
 class FNiagaraEffectEditorViewportClient : public FEditorViewportClient
diff --git a/Engine/Source/Editor/NiagaraEditor/Private/SNiagaraEffectEditorWidget.h b/Engine/Source/Editor/NiagaraEditor/Private/SNiagaraEffectEditorWidget.h
--- a/Engine/Source/Editor/NiagaraEditor/Private/SNiagaraEffectEditorWidget.h
+++ b/Engine/Source/Editor/NiagaraEditor/Private/SNiagaraEffectEditorWidget.h
@@ -10,8 +10,11 @@
 #include "SNameComboBox.h"
 #include "AssetThumbnail.h"
 #include "SNiagaraEffectEditorViewport.h"
-
-
+#include "Components/NiagaraComponent.h"
+#include "Engine/NiagaraEffect.h"
+#include "Engine/NiagaraSimulation.h"
+#include "Engine/NiagaraEffectRenderer.h"
+#include "ComponentReregisterContext.h"
 
 #define NGED_SECTION_BORDER SNew(SBorder).BorderImage(FEditorStyle::GetBrush("ToolPanel.GroupBorder")).Padding(4.0f).HAlign(HAlign_Left)
 
diff --git a/Engine/Source/Runtime/Engine/Classes/Engine/NiagaraEffectRenderer.h b/Engine/Source/Runtime/Engine/Classes/Engine/NiagaraEffectRenderer.h
--- a/Engine/Source/Runtime/Engine/Classes/Engine/NiagaraEffectRenderer.h
+++ b/Engine/Source/Runtime/Engine/Classes/Engine/NiagaraEffectRenderer.h
@@ -9,6 +9,7 @@ NiagaraEffectRenderer.h: Base class for Niagara render modules
 #include "ParticleHelper.h"
 #include "ParticleVertexFactory.h"
 #include "ParticleBeamTrailVertexFactory.h"
+#include "Components/NiagaraComponent.h"
 
 DECLARE_CYCLE_STAT(TEXT("Generate Sprite Vertex Data"), STAT_NiagaraGenSpriteVertexData, STATGROUP_Niagara);
 DECLARE_CYCLE_STAT(TEXT("Generate Ribbon Vertex Data"), STAT_NiagaraGenRibbonVertexData, STATGROUP_Niagara);
diff --git a/Engine/Source/Runtime/Engine/Private/InstancedStaticMesh.cpp b/Engine/Source/Runtime/Engine/Private/InstancedStaticMesh.cpp
--- a/Engine/Source/Runtime/Engine/Private/InstancedStaticMesh.cpp
+++ b/Engine/Source/Runtime/Engine/Private/InstancedStaticMesh.cpp
@@ -9,6 +9,7 @@
 #include "ShaderParameters.h"
 #include "ShaderParameterUtils.h"
 
+#include "Misc/UObjectToken.h"
 #include "Components/InteractiveFoliageComponent.h"
 #include "Components/SplineMeshComponent.h"
 #include "Components/ModelComponent.h"
diff --git a/Engine/Source/Runtime/Engine/Private/NiagaraEffectRenderer.cpp b/Engine/Source/Runtime/Engine/Private/NiagaraEffectRenderer.cpp
--- a/Engine/Source/Runtime/Engine/Private/NiagaraEffectRenderer.cpp
+++ b/Engine/Source/Runtime/Engine/Private/NiagaraEffectRenderer.cpp
@@ -1,5 +1,10 @@
+// Copyright 1998-2014 Epic Games, Inc. All Rights Reserved.
+
+
 #include "EnginePrivate.h"
 #include "Engine/NiagaraEffectRenderer.h"
+#include "Particles/ParticleResources.h"
+
 
 DECLARE_CYCLE_STAT(TEXT("PreRenderView"), STAT_NiagaraPreRenderView, STATGROUP_Niagara);
 DECLARE_CYCLE_STAT(TEXT("Render Total"), STAT_NiagaraRender, STATGROUP_Niagara);
diff --git a/Engine/Source/Runtime/Engine/Private/NiagaraSimulation.cpp b/Engine/Source/Runtime/Engine/Private/NiagaraSimulation.cpp
--- a/Engine/Source/Runtime/Engine/Private/NiagaraSimulation.cpp
+++ b/Engine/Source/Runtime/Engine/Private/NiagaraSimulation.cpp
@@ -1,6 +1,7 @@
 
 #include "EnginePrivate.h"
 #include "Engine/NiagaraSimulation.h"
+#include "Engine/NiagaraEffectRenderer.h"
 #include "VectorVM.h"
 
 
