Details panel for Niagara script editor.

diff --git a/Engine/Source/Editor/NiagaraEditor/Private/NiagaraEditor.cpp b/Engine/Source/Editor/NiagaraEditor/Private/NiagaraEditor.cpp
--- a/Engine/Source/Editor/NiagaraEditor/Private/NiagaraEditor.cpp
+++ b/Engine/Source/Editor/NiagaraEditor/Private/NiagaraEditor.cpp
@@ -7,32 +7,44 @@
 #include "Editor/WorkspaceMenuStructure/Public/WorkspaceMenuStructureModule.h"
 #include "SDockTab.h"
 #include "GenericCommands.h"
+
+#include "Editor/PropertyEditor/Public/PropertyEditorModule.h"
+#include "Editor/PropertyEditor/Public/IDetailsView.h"
  
 #define LOCTEXT_NAMESPACE "NiagaraEditor"
 
-const FName FNiagaraEditor::UpdateGraphTabId( TEXT( "NiagaraEditor_UpdateGraph" ) );
+const FName FNiagaraEditor::UpdateGraphTabId(TEXT("NiagaraEditor_UpdateGraph"));
+const FName FNiagaraEditor::PropertiesTabId(TEXT("NiagaraEditor_MaterialProperties"));
 
 void FNiagaraEditor::RegisterTabSpawners(const TSharedRef<class FTabManager>& TabManager)
 {
 	WorkspaceMenuCategory = TabManager->AddLocalWorkspaceMenuCategory(LOCTEXT("WorkspaceMenu_NiagaraEditor", "Niagara"));
 
 	FAssetEditorToolkit::RegisterTabSpawners(TabManager);
+	
+	TSharedRef<FWorkspaceItem> WorkspaceMenuCategoryRef = WorkspaceMenuCategory.ToSharedRef();
 
 	TabManager->RegisterTabSpawner( UpdateGraphTabId, FOnSpawnTab::CreateSP(this, &FNiagaraEditor::SpawnTab_UpdateGraph) )
 		.SetDisplayName( LOCTEXT("UpdateGraph", "Update Graph") )
-		.SetGroup( WorkspaceMenuCategory.ToSharedRef() );
+		.SetGroup(WorkspaceMenuCategoryRef);
+
+	TabManager->RegisterTabSpawner(PropertiesTabId, FOnSpawnTab::CreateSP(this, &FNiagaraEditor::SpawnTab_NodeProperties))
+		.SetDisplayName(LOCTEXT("DetailsTab", "Details"))
+		.SetGroup(WorkspaceMenuCategoryRef)
+		.SetIcon(FSlateIcon(FEditorStyle::GetStyleSetName(), "LevelEditor.Tabs.Details"));
 }
 
 void FNiagaraEditor::UnregisterTabSpawners(const TSharedRef<class FTabManager>& TabManager)
 {
 	FAssetEditorToolkit::UnregisterTabSpawners(TabManager);
 
-	TabManager->UnregisterTabSpawner( UpdateGraphTabId );
+	TabManager->UnregisterTabSpawner(UpdateGraphTabId);
+	TabManager->UnregisterTabSpawner(PropertiesTabId);
 }
 
 FNiagaraEditor::~FNiagaraEditor()
 {
-
+	NiagaraDetailsView.Reset();
 }
 
 
@@ -43,7 +55,7 @@ void FNiagaraEditor::InitNiagaraEditor( const EToolkitMode::Type Mode, const TSh
 	Source = CastChecked<UNiagaraScriptSource>(Script->Source);
 	check(Source->UpdateGraph != NULL);
 
-	TSharedRef<FTabManager::FLayout> StandaloneDefaultLayout = FTabManager::NewLayout( "Standalone_Niagara_Layout_v2" )
+	TSharedRef<FTabManager::FLayout> StandaloneDefaultLayout = FTabManager::NewLayout( "Standalone_Niagara_Layout_v3" )
 	->AddArea
 	(
 		FTabManager::NewPrimaryArea() ->SetOrientation(Orient_Vertical)
@@ -56,11 +68,26 @@ void FNiagaraEditor::InitNiagaraEditor( const EToolkitMode::Type Mode, const TSh
 		)
 		->Split
 		(
-			FTabManager::NewStack()
-			->AddTab( UpdateGraphTabId, ETabState::OpenedTab )
+			FTabManager::NewSplitter()->SetOrientation(Orient_Horizontal)->SetSizeCoefficient(0.9f)
+			->Split
+			(
+				FTabManager::NewStack()
+				->SetSizeCoefficient(0.2f)
+				->AddTab(PropertiesTabId, ETabState::OpenedTab)
+			)
+			->Split
+			(
+				FTabManager::NewStack()
+				->SetSizeCoefficient(0.8f)
+				->AddTab(UpdateGraphTabId, ETabState::OpenedTab)
+			)
 		)
 	);
 
+	FPropertyEditorModule& PropertyEditorModule = FModuleManager::GetModuleChecked<FPropertyEditorModule>("PropertyEditor");
+	const FDetailsViewArgs DetailsViewArgs(false, false, true, false, true, this);
+	NiagaraDetailsView = PropertyEditorModule.CreateDetailView(DetailsViewArgs);
+
 	const bool bCreateDefaultStandaloneMenu = true;
 	const bool bCreateDefaultToolbar = true;
 	FAssetEditorToolkit::InitAssetEditor( Mode, InitToolkitHost, FNiagaraEditorModule::NiagaraEditorAppIdentifier, StandaloneDefaultLayout, bCreateDefaultStandaloneMenu, bCreateDefaultToolbar, Script );
@@ -123,6 +150,7 @@ TSharedRef<SGraphEditor> FNiagaraEditor::CreateGraphEditorWidget(UEdGraph* InGra
 	AppearanceInfo.CornerText = LOCTEXT("AppearanceCornerText", "NIAGARA").ToString();
 
 	SGraphEditor::FGraphEditorEvents InEvents;
+	InEvents.OnSelectionChanged = SGraphEditor::FOnSelectionChanged::CreateSP(this, &FNiagaraEditor::OnSelectedNodesChanged);
 
 	// Make title bar
 	TSharedRef<SWidget> TitleBarWidget = 
@@ -167,6 +195,25 @@ TSharedRef<SDockTab> FNiagaraEditor::SpawnTab_UpdateGraph( const FSpawnTabArgs&
 		];
 }
 
+TSharedRef<SDockTab> FNiagaraEditor::SpawnTab_NodeProperties(const FSpawnTabArgs& Args)
+{
+	TSharedRef<SDockTab> SpawnedTab = SNew(SDockTab)
+		.Icon(FEditorStyle::GetBrush("LevelEditor.Tabs.Details"))
+		.Label(LOCTEXT("NiagaraDetailsTitle", "Details"))
+		[
+			NiagaraDetailsView.ToSharedRef()
+		];
+
+	TSharedPtr<SGraphEditor> UpdateGraphEditor = UpdateGraphEditorPtr.Pin();
+	if (UpdateGraphEditor.IsValid())
+	{
+		// Since we're initialising, make sure nothing is selected
+		UpdateGraphEditor->ClearSelectionSet();
+	}
+
+	return SpawnedTab;
+}
+
 void FNiagaraEditor::ExtendToolbar()
 {
 	struct Local
@@ -258,4 +305,25 @@ bool FNiagaraEditor::CanDeleteNodes() const
 	return (UpdateGraphEditor.IsValid() && UpdateGraphEditor->GetSelectedNodes().Num() > 0);
 }
 
+void FNiagaraEditor::OnSelectedNodesChanged(const TSet<class UObject*>& NewSelection)
+{
+	TArray<UObject*> SelectedObjects;
+
+	UObject* EditObject = Script;
+
+	if (NewSelection.Num() == 0)
+	{
+		SelectedObjects.Add(EditObject);
+	}
+	else
+	{
+		for (TSet<class UObject*>::TConstIterator SetIt(NewSelection); SetIt; ++SetIt)
+		{
+			SelectedObjects.Add(*SetIt);
+		}
+	}
+
+	GetDetailView()->SetObjects(SelectedObjects, true);
+}
+
 #undef LOCTEXT_NAMESPACE
diff --git a/Engine/Source/Editor/NiagaraEditor/Private/NiagaraEditor.h b/Engine/Source/Editor/NiagaraEditor/Private/NiagaraEditor.h
--- a/Engine/Source/Editor/NiagaraEditor/Private/NiagaraEditor.h
+++ b/Engine/Source/Editor/NiagaraEditor/Private/NiagaraEditor.h
@@ -7,7 +7,7 @@
 
 
 /** Viewer/editor for a DataTable */
-class FNiagaraEditor : public INiagaraEditor
+class FNiagaraEditor : public INiagaraEditor, public FNotifyHook
 {
 
 public:
@@ -33,6 +33,13 @@ public:
 	void DeleteSelectedNodes();
 	bool CanDeleteNodes() const;
 
+	// Widget Accessors
+	TSharedRef<class IDetailsView> GetDetailView() const { return NiagaraDetailsView.ToSharedRef(); }
+protected:
+
+	/** Called when the selection changes in the GraphEditor */
+	void OnSelectedNodesChanged(const TSet<class UObject*>& NewSelection);
+
 private:
 	/** Create widget for graph editing */
 	TSharedRef<class SGraphEditor> CreateGraphEditorWidget(UEdGraph* InGraph);
@@ -40,6 +47,9 @@ private:
 	/** Spawns the tab with the update graph inside */
 	TSharedRef<SDockTab> SpawnTab_UpdateGraph(const FSpawnTabArgs& Args);
 
+	/** Spawns the tab with the node details inside. */
+	TSharedRef<SDockTab> SpawnTab_NodeProperties(const FSpawnTabArgs& Args);
+
 	/** Builds the toolbar widget */
 	void ExtendToolbar();
 private:
@@ -54,6 +64,11 @@ private:
 	/** The command list for this editor */
 	TSharedPtr<FUICommandList> GraphEditorCommands;
 
-	/**	Graph editor tab */
+	/** Property View */
+	TSharedPtr<class IDetailsView> NiagaraDetailsView;
+
+	/**	The tab ids for the Niagara editor */
 	static const FName UpdateGraphTabId;
+	static const FName PropertiesTabId;
+
 };
