Niagara: exposing script constants to the effect editor, storing exposed constants on the emitter properties (themselves stored on the effect). This is one of the big features to enable reuse of scripts across different effects. Also contains bug fixes for ribbon emitters and a few other improvements.

diff --git a/Engine/Source/Editor/NiagaraEditor/Private/SNiagaraEffectEditorWidget.cpp b/Engine/Source/Editor/NiagaraEditor/Private/SNiagaraEffectEditorWidget.cpp
--- a/Engine/Source/Editor/NiagaraEditor/Private/SNiagaraEffectEditorWidget.cpp
+++ b/Engine/Source/Editor/NiagaraEditor/Private/SNiagaraEffectEditorWidget.cpp
@@ -123,6 +123,7 @@ void SNiagaraEffectEditorWidget::Construct(const FArguments& InArgs)
 						.ItemHeight(256)
 						.ListItemsSource(&(this->EffectInstance->GetEmitters()))
 						.OnGenerateRow(this, &SNiagaraEffectEditorWidget::OnGenerateWidgetForList)
+						.SelectionMode(ESelectionMode::None)
 					]
 
 					// Bottom-right corner text indicating the type of tool
@@ -187,6 +188,20 @@ void SEmitterWidget::Construct(const FArguments& InArgs)
 	RenderModuleList.Add(MakeShareable(new FString("Trails")));
 	RenderModuleList.Add(MakeShareable(new FString("Meshes")));
 
+	static TArray<TSharedPtr<EditorExposedVectorConstant>> DummyArray;
+	TArray<TSharedPtr<EditorExposedVectorConstant>> *EditorExposedConstants = &DummyArray;
+	if (Emitter->GetProperties()->UpdateScript && Emitter->GetProperties()->UpdateScript->Source)
+	{
+		EditorExposedConstants = &Emitter->GetProperties()->UpdateScript->Source->ExposedVectorConstants;
+	}
+
+	TArray<TSharedPtr<EditorExposedVectorConstant>> *EditorExposedSpawnConstants = &DummyArray;
+	if (Emitter->GetProperties()->SpawnScript && Emitter->GetProperties()->SpawnScript->Source)
+	{
+		EditorExposedSpawnConstants = &Emitter->GetProperties()->SpawnScript->Source->ExposedVectorConstants;
+	}
+
+
 	ChildSlot
 		.Padding(4)
 		[
@@ -242,7 +257,7 @@ void SEmitterWidget::Construct(const FArguments& InArgs)
 						//
 						+SHorizontalBox::Slot()
 						.Padding(4.0f)
-						.FillWidth(800)
+						.FillWidth(950)
 						.HAlign(HAlign_Left)
 						[
 							NGED_SECTION_BORDER
@@ -252,17 +267,19 @@ void SEmitterWidget::Construct(const FArguments& InArgs)
 								.Padding(4)
 								.AutoHeight()
 								[
-									SNew(SHorizontalBox)
-									+ SHorizontalBox::Slot()
-									.HAlign(HAlign_Right)
-									.VAlign(VAlign_Center)
-									.AutoWidth()
-									.Padding(4)
+									NGED_SECTION_DARKBORDER
 									[
-										SNew(STextBlock)
-										.Text(FString("Update Script"))
-									]
-									+ SHorizontalBox::Slot()
+										SNew(SHorizontalBox)
+										+ SHorizontalBox::Slot()
+										.HAlign(HAlign_Right)
+										.VAlign(VAlign_Center)
+										.AutoWidth()
+										.Padding(4)
+										[
+											SNew(STextBlock)
+											.Text(FString("Update Script"))
+										]
+										+ SHorizontalBox::Slot()
 										.HAlign(HAlign_Right)
 										.VAlign(VAlign_Center)
 										[
@@ -274,12 +291,27 @@ void SEmitterWidget::Construct(const FArguments& InArgs)
 											.OnSetReference(this, &SEmitterWidget::OnUpdateScriptSelectedFromPicker)
 											.WidthOverride(150)
 										]
+									]
+								]
 
+								+ SVerticalBox::Slot()
+								.Padding(4)
+								.AutoHeight()
+								.Expose(UpdateScriptConstantListSlot)
+								[
+									SAssignNew(UpdateScriptConstantList, SListView<TSharedPtr<EditorExposedVectorConstant>>)
+									.ItemHeight(20).ListItemsSource(EditorExposedConstants).OnGenerateRow(this, &SEmitterWidget::OnGenerateConstantListRow)
+									.SelectionMode(ESelectionMode::None)
 								]
 
+
+
+
 								+ SVerticalBox::Slot()
-									.Padding(4)
-									.AutoHeight()
+								.Padding(4)
+								.AutoHeight()
+								[
+									NGED_SECTION_DARKBORDER
 									[
 										SNew(SHorizontalBox)
 										+ SHorizontalBox::Slot()
@@ -292,17 +324,28 @@ void SEmitterWidget::Construct(const FArguments& InArgs)
 											.Text(FString("Spawn Script"))
 										]
 										+ SHorizontalBox::Slot()
-											.HAlign(HAlign_Right)
-											.VAlign(VAlign_Center)
-											[
-												SAssignNew(SpawnScriptSelector, SContentReference)
-												.AllowedClass(UNiagaraScript::StaticClass())
-												.AssetReference(this, &SEmitterWidget::GetSpawnScript)
-												.AllowSelectingNewAsset(true)
-												.AllowClearingReference(true)
-												.OnSetReference(this, &SEmitterWidget::OnSpawnScriptSelectedFromPicker)
-												.WidthOverride(150)
-											]
+										.HAlign(HAlign_Right)
+										.VAlign(VAlign_Center)
+										[
+											SAssignNew(SpawnScriptSelector, SContentReference)
+											.AllowedClass(UNiagaraScript::StaticClass())
+											.AssetReference(this, &SEmitterWidget::GetSpawnScript)
+											.AllowSelectingNewAsset(true)
+											.AllowClearingReference(true)
+											.OnSetReference(this, &SEmitterWidget::OnSpawnScriptSelectedFromPicker)
+											.WidthOverride(150)
+										]
+									]
+								]
+
+								+ SVerticalBox::Slot()
+									.Padding(4)
+									.AutoHeight()
+									.Expose(SpawnScriptConstantListSlot)
+									[
+										SAssignNew(SpawnScriptConstantList, SListView<TSharedPtr<EditorExposedVectorConstant>>)
+										.ItemHeight(20).ListItemsSource(EditorExposedSpawnConstants).OnGenerateRow(this, &SEmitterWidget::OnGenerateConstantListRow)
+										.SelectionMode(ESelectionMode::None)
 									]
 							]
 						]
@@ -430,4 +473,5 @@ void SEmitterWidget::Construct(const FArguments& InArgs)
 					]
 			]
 		];
+
 }
\ No newline at end of file
diff --git a/Engine/Source/Editor/NiagaraEditor/Private/SNiagaraEffectEditorWidget.h b/Engine/Source/Editor/NiagaraEditor/Private/SNiagaraEffectEditorWidget.h
--- a/Engine/Source/Editor/NiagaraEditor/Private/SNiagaraEffectEditorWidget.h
+++ b/Engine/Source/Editor/NiagaraEditor/Private/SNiagaraEffectEditorWidget.h
@@ -17,6 +17,78 @@
 #include "ComponentReregisterContext.h"
 
 #define NGED_SECTION_BORDER SNew(SBorder).BorderImage(FEditorStyle::GetBrush("ToolPanel.GroupBorder")).Padding(4.0f).HAlign(HAlign_Left)
+#define NGED_SECTION_LIGHTBORDER SNew(SBorder).BorderImage(FEditorStyle::GetBrush("ToolPanel.LightGroupBorder")).Padding(4.0f).HAlign(HAlign_Left)
+#define NGED_SECTION_DARKBORDER SNew(SBorder).BorderImage(FEditorStyle::GetBrush("ToolPanel.DarkGroupBorder")).Padding(4.0f).HAlign(HAlign_Left)
+
+
+class SVectorConstantWidget : public SCompoundWidget, public FNotifyHook
+{
+private:	
+	FName ConstantName;
+	TSharedPtr<FNiagaraSimulation> Emitter;
+
+	TSharedPtr<SEditableTextBox> XText, YText, ZText, WText;
+
+public:
+	SLATE_BEGIN_ARGS(SVectorConstantWidget) :
+	_Emitter(nullptr),
+	_ConstantName("Undefined")
+	{
+	}
+	SLATE_ARGUMENT(TSharedPtr<FNiagaraSimulation>, Emitter)
+	SLATE_ARGUMENT(FName, ConstantName)
+	SLATE_END_ARGS()
+
+
+	void OnConstantChanged(const FText &InText)
+	{
+		FVector4 Vec;
+		FString Strns[4];
+		Strns[0] = XText->GetText().ToString();
+		Strns[1] = YText->GetText().ToString();
+		Strns[2] = ZText->GetText().ToString();
+		Strns[3] = WText->GetText().ToString();
+		for (uint32 i = 0; i < 4; i++)
+		{
+			if (Strns[i]!="")
+			{
+				Vec[i] = FCString::Atof(*Strns[i]);
+			}
+			else
+			{
+				Vec[i] = 0.0f;
+			}
+		}
+		Emitter->GetProperties()->ExternalConstants.SetOrAdd(ConstantName, Vec);
+	}
+
+	void Construct(const FArguments& InArgs)
+	{
+		Emitter = InArgs._Emitter;
+		ConstantName = InArgs._ConstantName;
+
+		FVector4 *ConstVal = Emitter->GetProperties()->ExternalConstants.FindVector(ConstantName);
+		FString ConstText[4] = {"","","",""};
+		if (ConstVal)
+		{
+			for (uint32 i = 0; i < 4; i++)
+			{
+				ConstText[i] = FString::Printf(TEXT("%.4f"), (*ConstVal)[i]);
+			}
+		}
+		
+		this->ChildSlot
+		[
+		SNew(SHorizontalBox)
+		+ SHorizontalBox::Slot().Padding(4)[SNew(STextBlock).Text(FText::FromName(ConstantName))]
+		+ SHorizontalBox::Slot().Padding(4)[SAssignNew(XText, SEditableTextBox).Text(FText::FromString(ConstText[0])).OnTextChanged(this, &SVectorConstantWidget::OnConstantChanged)]
+		+ SHorizontalBox::Slot().Padding(4)[SAssignNew(YText, SEditableTextBox).Text(FText::FromString(ConstText[1])).OnTextChanged(this, &SVectorConstantWidget::OnConstantChanged)]
+		+ SHorizontalBox::Slot().Padding(4)[SAssignNew(ZText, SEditableTextBox).Text(FText::FromString(ConstText[2])).OnTextChanged(this, &SVectorConstantWidget::OnConstantChanged)]
+		+ SHorizontalBox::Slot().Padding(4)[SAssignNew(WText, SEditableTextBox).Text(FText::FromString(ConstText[3])).OnTextChanged(this, &SVectorConstantWidget::OnConstantChanged)]
+		];
+		
+	}
+};
 
 
 class SEmitterWidget : public SCompoundWidget, public FNotifyHook
@@ -35,6 +107,11 @@ private:
 
 	TArray<TSharedPtr<FString> > RenderModuleList;
 
+	TSharedPtr< SListView<TSharedPtr<EditorExposedVectorConstant>>> UpdateScriptConstantList;
+	SVerticalBox::FSlot *UpdateScriptConstantListSlot;
+	TSharedPtr< SListView<TSharedPtr<EditorExposedVectorConstant>>> SpawnScriptConstantList;
+	SVerticalBox::FSlot *SpawnScriptConstantListSlot;
+
 public:
 	SLATE_BEGIN_ARGS(SEmitterWidget) :
 		_Emitter(nullptr)
@@ -44,6 +121,16 @@ public:
 	SLATE_ARGUMENT(FNiagaraEffectInstance*, Effect)
 	SLATE_END_ARGS()
 
+	TSharedRef<ITableRow> OnGenerateConstantListRow(TSharedPtr<EditorExposedVectorConstant> InItem, const TSharedRef< STableViewBase >& OwnerTable)
+	{
+		return SNew(STableRow<TSharedPtr<FString>>, OwnerTable)
+		.Content()
+		[
+			SNew(SVectorConstantWidget).ConstantName(InItem->ConstName).Emitter(Emitter)
+		];
+	}
+
+
 	void OnUpdateScriptSelectedFromPicker(UObject *Asset)
 	{
 		// Close the asset picker
@@ -52,6 +139,13 @@ public:
 		// Set the object found from the asset picker
 		CurUpdateScript = Cast<UNiagaraScript>(Asset);
 		Emitter->GetProperties()->UpdateScript = CurUpdateScript;
+
+		UpdateScriptConstantListSlot->DetachWidget();
+		(*UpdateScriptConstantListSlot)
+			[
+				SAssignNew(UpdateScriptConstantList, SListView<TSharedPtr<EditorExposedVectorConstant>>)
+				.ItemHeight(20).ListItemsSource(&CurUpdateScript->Source->ExposedVectorConstants).OnGenerateRow(this, &SEmitterWidget::OnGenerateConstantListRow)
+			];
 	}
 
 	void OnSpawnScriptSelectedFromPicker(UObject *Asset)
diff --git a/Engine/Source/Editor/UnrealEd/Classes/Niagara/NiagaraNodeConstant.h b/Engine/Source/Editor/UnrealEd/Classes/Niagara/NiagaraNodeConstant.h
--- a/Engine/Source/Editor/UnrealEd/Classes/Niagara/NiagaraNodeConstant.h
+++ b/Engine/Source/Editor/UnrealEd/Classes/Niagara/NiagaraNodeConstant.h
@@ -22,6 +22,11 @@ public:
 	UPROPERTY(EditAnywhere, Category = "Constant")
 	uint32 bNeedsDefault : 1;
 
+	UPROPERTY(EditAnywhere, Category = "Constant")
+	uint32 bExposeToEffectEditor : 1;
+
+	bool IsExposedToEffectEditor()	{ return bExposeToEffectEditor; }
+
 	// Begin UObject interface
 	virtual void PostEditChangeProperty(struct FPropertyChangedEvent& PropertyChangedEvent);
 	// End UObject interface
diff --git a/Engine/Source/Editor/UnrealEd/Classes/Niagara/NiagaraScriptSource.h b/Engine/Source/Editor/UnrealEd/Classes/Niagara/NiagaraScriptSource.h
--- a/Engine/Source/Editor/UnrealEd/Classes/Niagara/NiagaraScriptSource.h
+++ b/Engine/Source/Editor/UnrealEd/Classes/Niagara/NiagaraScriptSource.h
@@ -5,6 +5,7 @@
 #include "Engine/NiagaraScriptSourceBase.h"
 #include "NiagaraScriptSource.generated.h"
 
+
 UCLASS(MinimalAPI)
 class UNiagaraScriptSource : public UNiagaraScriptSourceBase
 {
@@ -16,6 +17,7 @@ class UNiagaraScriptSource : public UNiagaraScriptSourceBase
 
 	// UObject interface.
 	virtual void PostLoad() override;
+	void Compile();
 
 	UNREALED_API void GetParticleAttributes(TArray<FName>& VectorOutputs);
 	UNREALED_API void GetEmitterAttributes(TArray<FName>& VectorInputs, TArray<FName>& MatrixInputs);
diff --git a/Engine/Source/Editor/UnrealEd/Private/Niagara/NiagaraNodeConstant.cpp b/Engine/Source/Editor/UnrealEd/Private/Niagara/NiagaraNodeConstant.cpp
--- a/Engine/Source/Editor/UnrealEd/Private/Niagara/NiagaraNodeConstant.cpp
+++ b/Engine/Source/Editor/UnrealEd/Private/Niagara/NiagaraNodeConstant.cpp
@@ -11,6 +11,7 @@ UNiagaraNodeConstant::UNiagaraNodeConstant(const FObjectInitializer& ObjectIniti
 	ConstName = NAME_None;
 	DataType = ENiagaraDataType::Vector;
 	bNeedsDefault = true;
+	bExposeToEffectEditor = false;
 }
 
 void UNiagaraNodeConstant::AllocateDefaultPins()
diff --git a/Engine/Source/Editor/UnrealEd/Private/Niagara/NiagaraScriptSource.cpp b/Engine/Source/Editor/UnrealEd/Private/Niagara/NiagaraScriptSource.cpp
--- a/Engine/Source/Editor/UnrealEd/Private/Niagara/NiagaraScriptSource.cpp
+++ b/Engine/Source/Editor/UnrealEd/Private/Niagara/NiagaraScriptSource.cpp
@@ -36,8 +36,31 @@ void UNiagaraScriptSource::PostLoad()
 	if (ScriptOwner)
 	{
 		ScriptOwner->ConditionalPostLoad();
-		FNiagaraEditorModule& NiagaraEditorModule = FModuleManager::Get().LoadModuleChecked<FNiagaraEditorModule>(TEXT("NiagaraEditor"));
-		NiagaraEditorModule.CompileScript(ScriptOwner);
+		//FNiagaraEditorModule& NiagaraEditorModule = FModuleManager::Get().LoadModuleChecked<FNiagaraEditorModule>(TEXT("NiagaraEditor"));
+		//NiagaraEditorModule.CompileScript(ScriptOwner);
+		Compile();
+	}
+
+}
+
+
+void UNiagaraScriptSource::Compile()
+{
+	UNiagaraScript* ScriptOwner = Cast<UNiagaraScript>(GetOuter());
+	FNiagaraEditorModule& NiagaraEditorModule = FModuleManager::Get().LoadModuleChecked<FNiagaraEditorModule>(TEXT("NiagaraEditor"));
+	NiagaraEditorModule.CompileScript(ScriptOwner);
+
+	// grab all constant nodes that are exposed to the editor
+	TArray<UNiagaraNodeConstant*> ConstNodes;
+	UpdateGraph->GetNodesOfClass<UNiagaraNodeConstant>(ConstNodes);
+	for (UNiagaraNodeConstant *Node : ConstNodes)
+	{
+		if (Node->bExposeToEffectEditor)
+		{
+			EditorExposedVectorConstant *Const = new EditorExposedVectorConstant();
+			Const->ConstName = Node->ConstName;
+			ExposedVectorConstants.Add( MakeShareable(Const) );
+		}
 	}
 }
 
diff --git a/Engine/Source/Runtime/Engine/Classes/Engine/NiagaraEffect.h b/Engine/Source/Runtime/Engine/Classes/Engine/NiagaraEffect.h
--- a/Engine/Source/Runtime/Engine/Classes/Engine/NiagaraEffect.h
+++ b/Engine/Source/Runtime/Engine/Classes/Engine/NiagaraEffect.h
@@ -13,12 +13,12 @@ class UNiagaraEffect : public UObject
 public:
 	ENGINE_API FNiagaraEmitterProperties &AddEmitterProperties();
 
-
 	void Init(UNiagaraComponent *InComponent)
 	{
 		Component = InComponent;
 	}
 
+
 	TArray<FNiagaraEmitterProperties> &GetEmitterProperties() {
 		return EmitterProps;
 	};
@@ -102,6 +102,7 @@ public:
 		for (TSharedPtr<FNiagaraSimulation>&it : Emitters)
 		{
 			it->SetConstants(Constants);
+			it->GetConstants().Merge(it->GetProperties()->ExternalConstants);
 			it->Tick(DeltaSeconds);
 		}
 	}
diff --git a/Engine/Source/Runtime/Engine/Classes/Engine/NiagaraScriptSourceBase.h b/Engine/Source/Runtime/Engine/Classes/Engine/NiagaraScriptSourceBase.h
--- a/Engine/Source/Runtime/Engine/Classes/Engine/NiagaraScriptSourceBase.h
+++ b/Engine/Source/Runtime/Engine/Classes/Engine/NiagaraScriptSourceBase.h
@@ -3,9 +3,18 @@
 #pragma once
 #include "NiagaraScriptSourceBase.generated.h"
 
+
+struct EditorExposedVectorConstant
+{
+	FName ConstName;
+	FVector4 Value;
+};
+
 /** Runtime data for a Niagara system */
 UCLASS(MinimalAPI)
 class UNiagaraScriptSourceBase : public UObject
 {
 	GENERATED_UCLASS_BODY()
+
+	TArray<TSharedPtr<EditorExposedVectorConstant> > ExposedVectorConstants;
 };
\ No newline at end of file
diff --git a/Engine/Source/Runtime/Engine/Classes/Engine/NiagaraSimulation.h b/Engine/Source/Runtime/Engine/Classes/Engine/NiagaraSimulation.h
--- a/Engine/Source/Runtime/Engine/Classes/Engine/NiagaraSimulation.h
+++ b/Engine/Source/Runtime/Engine/Classes/Engine/NiagaraSimulation.h
@@ -53,6 +53,12 @@ public:
 	UMaterial *Material;
 	UPROPERTY(EditAnywhere, Category = "Emitter Properties")
 	TEnumAsByte<EEmitterRenderModuleType> RenderModuleType;
+
+	UPROPERTY(EditAnywhere, Category = "Emitter Properties")
+	FNiagaraConstantMap ExternalConstants;		// these are the update script constants from the effect editor; will be added to the emitter's constant map
+
+	UPROPERTY(EditAnywhere, Category = "Emitter Properties")
+	FNiagaraConstantMap ExternalSpawnConstants;		// these are the spawn script constants from the effect editor; will be added to the emitter's constant map
 };
 
 
@@ -77,7 +83,9 @@ public:
 	FNiagaraEmitterParticleData &GetData()	{ return Data; }
 
 	void SetConstants(const FNiagaraConstantMap &InMap)	{ Constants = InMap; }
-	NiagaraEffectRenderer *GetEffectRenderer()	{ return EffectRenderer;  }
+	FNiagaraConstantMap &GetConstants()	{ return Constants; }
+
+	NiagaraEffectRenderer *GetEffectRenderer()	{ return EffectRenderer; }
 	
 	bool IsEnabled()	{ return bIsEnabled;  }
 	void SetEnabled(bool bInEnabled)	{ bIsEnabled = bInEnabled;  }
@@ -88,9 +96,8 @@ public:
 
 	FNiagaraEmitterProperties *GetProperties()	{ return Props; }
 	void SetProperties(FNiagaraEmitterProperties *InProps)	{ Props = InProps; }
-
 private:
-	FNiagaraEmitterProperties *Props;
+	FNiagaraEmitterProperties *Props;		// points to an entry in the array of FNiagaraProperties stored in the EffectInstance (itself pointing to the effect's properties)
 
 	float Age;
 	bool bIsEnabled;
diff --git a/Engine/Source/Runtime/Engine/Classes/Niagara/NiagaraConstantSet.h b/Engine/Source/Runtime/Engine/Classes/Niagara/NiagaraConstantSet.h
--- a/Engine/Source/Runtime/Engine/Classes/Niagara/NiagaraConstantSet.h
+++ b/Engine/Source/Runtime/Engine/Classes/Niagara/NiagaraConstantSet.h
@@ -5,13 +5,13 @@
 #include "Niagara/NiagaraCommon.h"
 #include "NiagaraConstantSet.generated.h"
 
+
 USTRUCT()
 struct FNiagaraConstantMap
 {
 	GENERATED_USTRUCT_BODY()
 
 private:
-
 	TMap<FName, float> ScalarConstants;
 	TMap<FName, FVector4> VectorConstants;
 	TMap<FName, FMatrix> MatrixConstants;
@@ -39,8 +39,36 @@ public:
 	const float* FindScalar(FName InName)const { return ScalarConstants.Find(InName); }
 	const FVector4* FindVector(FName InName)const { return VectorConstants.Find(InName); }
 	const FMatrix* FindMatrix(FName InName)const { return MatrixConstants.Find(InName); }
+
+	void Merge(FNiagaraConstantMap &InMap)
+	{
+		for (auto MapIt = InMap.ScalarConstants.CreateIterator(); MapIt; ++MapIt)
+		{
+			SetOrAdd(MapIt.Key(), MapIt.Value());
+		}
+
+		for (auto MapIt = InMap.VectorConstants.CreateIterator(); MapIt; ++MapIt)
+		{
+			SetOrAdd(MapIt.Key(), MapIt.Value());
+		}
+
+		for (auto MapIt = InMap.MatrixConstants.CreateIterator(); MapIt; ++MapIt)
+		{
+			SetOrAdd(MapIt.Key(), MapIt.Value());
+		}
+	}
+
+
+	virtual bool Serialize(FArchive &Ar)
+	{
+		Ar << ScalarConstants << VectorConstants << MatrixConstants;
+		return true;
+	}
 };
 
+
+
+
 USTRUCT()
 struct FNiagaraConstants
 {
@@ -112,12 +140,14 @@ public:
 		}
 
 		Idx += FMath::Min(1, NewIdx % 4);//Move to the next table entry if needed.
+
 		for (int32 i = 0; i < VectorNames.Num(); ++i)
 		{
 			FName VcName = VectorNames[i];
 			const FVector4* Vector = Externals.FindVector(VcName);
 			ConstantsTable[Idx++] = Vector ? *Vector : VectorConstants[i];
 		}
+
 		for (int32 i = 0; i < MatrixNames.Num(); ++i)
 		{
 			FName McName = MatrixNames[i];
@@ -142,6 +172,7 @@ public:
 		{
 			ScalarConstants[Idx] = Sc;
 		}
+
 	}
 
 	void SetOrAdd(FName Name, const FVector4& Vc)
diff --git a/Engine/Source/Runtime/Engine/Private/NiagaraConstantSet.cpp b/Engine/Source/Runtime/Engine/Private/NiagaraConstantSet.cpp
new file mode 100644
--- /dev/null
+++ b/Engine/Source/Runtime/Engine/Private/NiagaraConstantSet.cpp
@@ -0,0 +1,11 @@
+#include "EnginePrivate.h"
+
+template<>
+struct TStructOpsTypeTraits<FNiagaraConstantMap> : public TStructOpsTypeTraitsBase
+{
+	enum
+	{
+		WithSerializer = true
+	};
+};
+IMPLEMENT_STRUCT(NiagaraConstantMap);
diff --git a/Engine/Source/Runtime/Engine/Private/NiagaraEffectRenderer.cpp b/Engine/Source/Runtime/Engine/Private/NiagaraEffectRenderer.cpp
--- a/Engine/Source/Runtime/Engine/Private/NiagaraEffectRenderer.cpp
+++ b/Engine/Source/Runtime/Engine/Private/NiagaraEffectRenderer.cpp
@@ -215,7 +215,12 @@ void NiagaraEffectRendererSprites::GetDynamicMeshElements(const TArray<const FSc
 	SCOPE_CYCLE_COUNTER(STAT_NiagaraRender);
 	SCOPE_CYCLE_COUNTER(STAT_NiagaraRenderSprites);
 
-	check(DynamicDataRender && DynamicDataRender->VertexData.Num());
+	check(DynamicDataRender)
+		
+	if (DynamicDataRender->VertexData.Num() == 0)
+	{
+		return;
+	}
 
 	const bool bIsWireframe = ViewFamily.EngineShowFlags.Wireframe;
 	FMaterialRenderProxy* MaterialRenderProxy = Material->GetRenderProxy(SceneProxy->IsSelected(), SceneProxy->IsHovered());
@@ -417,7 +422,11 @@ void NiagaraEffectRendererRibbon::GetDynamicMeshElements(const TArray<const FSce
 	SCOPE_CYCLE_COUNTER(STAT_NiagaraRender);
 	SCOPE_CYCLE_COUNTER(STAT_NiagaraRenderRibbons);
 
-	check(DynamicDataRender && DynamicDataRender->VertexData.Num());
+	check(DynamicDataRender)
+	if (DynamicDataRender->VertexData.Num() == 0)
+	{
+		return;
+	}
 
 	const bool bIsWireframe = ViewFamily.EngineShowFlags.Wireframe;
 	FMaterialRenderProxy* MaterialRenderProxy = Material->GetRenderProxy(SceneProxy->IsSelected(), SceneProxy->IsHovered());
@@ -466,8 +475,8 @@ void NiagaraEffectRendererRibbon::GetDynamicMeshElements(const TArray<const FSce
 				CollectorResources.UniformBuffer = FParticleBeamTrailUniformBufferRef::CreateUniformBufferImmediate(PerViewUniformParameters, UniformBuffer_SingleFrame);
 
 				CollectorResources.VertexFactory.InitResource();
-				CollectorResources.VertexFactory.SetBeamTrailUniformBuffer(PerViewUniformBuffers[ViewIndex]);
-				CollectorResources.VertexFactory.SetVertexBuffer(DynamicVertexAllocation.VertexBuffer, DynamicVertexAllocation.VertexOffset, sizeof(FParticleBeamTrailVertex));
+				CollectorResources.VertexFactory.SetBeamTrailUniformBuffer(CollectorResources.UniformBuffer);
+				CollectorResources.VertexFactory.SetVertexBuffer(LocalDynamicVertexAllocation.VertexBuffer, LocalDynamicVertexAllocation.VertexOffset, sizeof(FParticleBeamTrailVertex));
 				CollectorResources.VertexFactory.SetDynamicParameterBuffer(NULL, 0, 0);
 
 				FMeshBatch& MeshBatch = Collector.AllocateMesh();
@@ -493,7 +502,7 @@ void NiagaraEffectRendererRibbon::GetDynamicMeshElements(const TArray<const FSce
 				FMeshBatchElement& MeshElement = MeshBatch.Elements[0];
 				MeshElement.IndexBuffer = &GParticleIndexBuffer;
 				MeshElement.FirstIndex = 0;
-				MeshElement.NumPrimitives = (DynamicDataRender->VertexData.Num() - 1) / 2;
+				MeshElement.NumPrimitives = (DynamicDataRender->VertexData.Num() - 2);
 				MeshElement.NumInstances = 1;
 				MeshElement.MinVertexIndex = 0;
 				MeshElement.MaxVertexIndex = DynamicDataRender->VertexData.Num() - 1;
@@ -588,7 +597,7 @@ void NiagaraEffectRendererRibbon::DrawDynamicElements(FPrimitiveDrawInterface* P
 		FMeshBatchElement& MeshElement = MeshBatch.Elements[0];
 		MeshElement.IndexBuffer = &GParticleIndexBuffer;
 		MeshElement.FirstIndex = 0;
-		MeshElement.NumPrimitives = DynamicDataRender->VertexData.Num() / 2;
+		MeshElement.NumPrimitives = DynamicDataRender->VertexData.Num()-2;
 		MeshElement.NumInstances = 1;
 		MeshElement.MinVertexIndex = 0;
 		MeshElement.MaxVertexIndex = DynamicDataRender->VertexData.Num() - 1;
@@ -699,20 +708,21 @@ FNiagaraDynamicDataBase *NiagaraEffectRendererRibbon::GenerateVertexData(const F
 
 		FVector ParticleRight = FVector::CrossProduct(ParticleDir, FVector(0.0f, 0.0f, 1.0f));
 		ParticleRight.Normalize();
+		ParticleRight *= 3.0f;
 
 		if (i == 0)
 		{
-			AddRibbonVert(RenderData, ParticlePos + ParticleRight, Data, UVs[0], ColorPtr[i], AgePtr[i], RotPtr[i]);
-			AddRibbonVert(RenderData, ParticlePos - ParticleRight, Data, UVs[1], ColorPtr[i], AgePtr[i], RotPtr[i]);
+			AddRibbonVert(RenderData, ParticlePos + ParticleRight, Data, UVs[0], ColorPtr[Index1], AgePtr[Index1], RotPtr[i]);
+			AddRibbonVert(RenderData, ParticlePos - ParticleRight, Data, UVs[1], ColorPtr[Index1], AgePtr[Index1], RotPtr[i]);
 		}
 		else
 		{
-			AddRibbonVert(RenderData, PrevPos2, Data, UVs[0], ColorPtr[i], AgePtr[i], RotPtr[i]);
-			AddRibbonVert(RenderData, PrevPos, Data, UVs[1], ColorPtr[i], AgePtr[i], RotPtr[i]);
+			AddRibbonVert(RenderData, PrevPos2, Data, UVs[0], ColorPtr[Index1], AgePtr[Index1], RotPtr[i]);
+			AddRibbonVert(RenderData, PrevPos, Data, UVs[1], ColorPtr[Index1], AgePtr[Index1], RotPtr[i]);
 		}
 
-		AddRibbonVert(RenderData, ParticlePos - ParticleRight + ParticleDir, Data, UVs[2], ColorPtr[i], AgePtr[i], RotPtr[i]);
-		AddRibbonVert(RenderData, ParticlePos + ParticleRight + ParticleDir, Data, UVs[3], ColorPtr[i], AgePtr[i], RotPtr[i]);
+		AddRibbonVert(RenderData, ParticlePos - ParticleRight + ParticleDir, Data, UVs[2], ColorPtr[Index2], AgePtr[Index2], RotPtr[i]);
+		AddRibbonVert(RenderData, ParticlePos + ParticleRight + ParticleDir, Data, UVs[3], ColorPtr[Index2], AgePtr[Index2], RotPtr[i]);
 		PrevPos = ParticlePos - ParticleRight + ParticleDir;
 		PrevPos2 = ParticlePos + ParticleRight + ParticleDir;
 		PrevDir = ParticleDir;
diff --git a/Engine/Source/Runtime/Engine/Private/NiagaraSimulation.cpp b/Engine/Source/Runtime/Engine/Private/NiagaraSimulation.cpp
--- a/Engine/Source/Runtime/Engine/Private/NiagaraSimulation.cpp
+++ b/Engine/Source/Runtime/Engine/Private/NiagaraSimulation.cpp
@@ -13,6 +13,7 @@ FNiagaraSimulation::FNiagaraSimulation(FNiagaraEmitterProperties *InProps)
 , EffectRenderer(nullptr)
 {
 	Props = InProps;
+	Constants.Merge(InProps->ExternalConstants);
 }
 
 FNiagaraSimulation::FNiagaraSimulation(FNiagaraEmitterProperties *InProps, ERHIFeatureLevel::Type InFeatureLevel)
@@ -23,6 +24,7 @@ FNiagaraSimulation::FNiagaraSimulation(FNiagaraEmitterProperties *InProps, ERHIF
 	, EffectRenderer(nullptr)
 {
 	Props = InProps;
+	Constants.Merge(InProps->ExternalConstants);
 	SetRenderModuleType(InProps->RenderModuleType, InFeatureLevel);
 }
 
