Niagara:
-adding renderer configuration
-fixing various bugs, including crashes, failed rendering, and weird behavior
-some additional UI love, making preview viewport work properly, more consistent updating of in-level effects after changes, immediate constant display in the effect editor on script compile for both update and spawn scripts
-changed sin/cos/tan nodes to use [0:1] instead of [0:Pi] for half a period (most range restricted values we deal with are [0:1], so it makes more sense to scale implicitly)
-added less than comparison operator; outputs 0 if false, 1 if true. Still experimental, may change

diff --git a/Engine/Source/Editor/NiagaraEditor/Private/NiagaraCompiler_VectorVM.cpp b/Engine/Source/Editor/NiagaraEditor/Private/NiagaraCompiler_VectorVM.cpp
--- a/Engine/Source/Editor/NiagaraEditor/Private/NiagaraCompiler_VectorVM.cpp
+++ b/Engine/Source/Editor/NiagaraEditor/Private/NiagaraCompiler_VectorVM.cpp
@@ -588,4 +588,10 @@ void FNiagaraCompiler_VectorVM::Inverse_Internal(TArray<TNiagaraExprPtr>& InputE
 	
 }
 
+
+void FNiagaraCompiler_VectorVM::LessThan_Internal(TArray<TNiagaraExprPtr>& InputExpressions, TArray<TNiagaraExprPtr>& OutputExpressions)
+{
+	OutputExpressions.Add(Expression_VMNative(VectorVM::EOp::lessthan, InputExpressions));
+}
+
 #undef LOCTEXT_NAMESPACE
diff --git a/Engine/Source/Editor/NiagaraEditor/Private/NiagaraEffectEditor.cpp b/Engine/Source/Editor/NiagaraEditor/Private/NiagaraEffectEditor.cpp
--- a/Engine/Source/Editor/NiagaraEditor/Private/NiagaraEffectEditor.cpp
+++ b/Engine/Source/Editor/NiagaraEditor/Private/NiagaraEffectEditor.cpp
@@ -208,8 +208,9 @@ FReply FNiagaraEffectEditor::OnAddEmitterClicked()
 {
 	FNiagaraEditorModule& NiagaraEditorModule = FModuleManager::LoadModuleChecked<FNiagaraEditorModule>("NiagaraEditor");
 
-	FNiagaraEmitterProperties &Props = Effect->AddEmitterProperties();
-	EffectInstance->AddEmitter(&Props);
+	FNiagaraEmitterProperties *Props = Effect->AddEmitterProperties();
+	TSharedPtr<FNiagaraSimulation> NewEmitter = EffectInstance->AddEmitter(Props);
+	Effect->CreateEffectRendererProps(NewEmitter);
 	UpdateEditorPtr.Pin()->GetViewport()->SetPreviewEffect(EffectInstance);
 	UpdateEditorPtr.Pin()->UpdateList();
 	return FReply::Handled();
diff --git a/Engine/Source/Editor/NiagaraEditor/Private/SNiagaraEffectEditorViewport.cpp b/Engine/Source/Editor/NiagaraEditor/Private/SNiagaraEffectEditorViewport.cpp
--- a/Engine/Source/Editor/NiagaraEditor/Private/SNiagaraEffectEditorViewport.cpp
+++ b/Engine/Source/Editor/NiagaraEditor/Private/SNiagaraEffectEditorViewport.cpp
@@ -150,7 +150,10 @@ void SNiagaraEffectEditorViewport::SetPreviewEffect(FNiagaraEffectInstance *InPr
 	PreviewScene.RemoveComponent(PreviewComponent);
 	PreviewScene.AddComponent(PreviewComponent, Transform);
 
-	FComponentReregisterContext ReregisterContext(PreviewComponent);
+	{
+		FComponentReregisterContext ReregisterContext(PreviewComponent);
+	}
+
 	PreviewComponent->SetEffectInstance(InPreviewEffect);
 	PreviewComponent->GetEffectInstance()->Init(PreviewComponent);
 }
diff --git a/Engine/Source/Editor/NiagaraEditor/Private/SNiagaraEffectEditorViewport.h b/Engine/Source/Editor/NiagaraEditor/Private/SNiagaraEffectEditorViewport.h
--- a/Engine/Source/Editor/NiagaraEditor/Private/SNiagaraEffectEditorViewport.h
+++ b/Engine/Source/Editor/NiagaraEditor/Private/SNiagaraEffectEditorViewport.h
@@ -47,7 +47,7 @@ public:
 	bool IsTogglePreviewGridChecked() const;
 	void TogglePreviewBackground();
 	bool IsTogglePreviewBackgroundChecked() const;
-	
+	class UNiagaraComponent *GetPreviewComponent()	{ return PreviewComponent;  }
 protected:
 	/** SEditorViewport interface */
 	virtual TSharedRef<FEditorViewportClient> MakeEditorViewportClient() override;
diff --git a/Engine/Source/Editor/NiagaraEditor/Private/SNiagaraEffectEditorWidget.cpp b/Engine/Source/Editor/NiagaraEditor/Private/SNiagaraEffectEditorWidget.cpp
--- a/Engine/Source/Editor/NiagaraEditor/Private/SNiagaraEffectEditorWidget.cpp
+++ b/Engine/Source/Editor/NiagaraEditor/Private/SNiagaraEffectEditorWidget.cpp
@@ -3,6 +3,12 @@
 #include "NiagaraEditorPrivatePCH.h"
 #include "SNiagaraEffectEditorWidget.h"
 #include "SNotificationList.h"
+#include "SExpandableArea.h"
+#include "Editor/PropertyEditor/Public/PropertyEditorModule.h"
+#include "Editor/PropertyEditor/Public/IStructureDetailsView.h"
+
+
+#define LOCTEXT_NAMESPACE "NiagaraEffectEditor"
 
 void SNiagaraEffectEditorWidget::Construct(const FArguments& InArgs)
 {
@@ -144,7 +150,6 @@ void SNiagaraEffectEditorWidget::Construct(const FArguments& InArgs)
 
 
 
-
 void SEmitterWidget::Construct(const FArguments& InArgs)
 {
 	CurSpawnScript = nullptr;
@@ -178,6 +183,16 @@ void SEmitterWidget::Construct(const FArguments& InArgs)
 	}
 
 
+	FPropertyEditorModule& PropertyModule = FModuleManager::LoadModuleChecked<FPropertyEditorModule>("PropertyEditor");
+	FDetailsViewArgs ViewArgs;
+	ViewArgs.bAllowSearch = false;
+	ViewArgs.bHideSelectionTip = true;
+	ViewArgs.bShowActorLabel = false;
+	TSharedRef<IDetailsView> RendererPropertiesViewRef = PropertyModule.CreateDetailView(ViewArgs);
+	RendererPropertiesView = RendererPropertiesViewRef;
+	RendererPropertiesView->SetObject(Emitter->GetProperties()->RendererProperties);
+
+	
 	ChildSlot
 		.Padding(4)
 		[
@@ -207,9 +222,10 @@ void SEmitterWidget::Construct(const FArguments& InArgs)
 							.FillWidth(0.3)
 							.HAlign(HAlign_Left)
 							[
-								SNew(SEditableText).Text(FText::FromString("Emitter"))
+								SNew(SEditableText).Text(this, &SEmitterWidget::GetEmitterName)
 								.MinDesiredWidth(250)
 								.Font(FEditorStyle::GetFontStyle("ContentBrowser.AssetListViewNameFontDirty"))
+								.OnTextChanged(this, &SEmitterWidget::OnEmitterNameChanged)
 							]
 						+ SHorizontalBox::Slot()
 							.HAlign(HAlign_Right)
@@ -238,91 +254,106 @@ void SEmitterWidget::Construct(const FArguments& InArgs)
 						[
 							NGED_SECTION_BORDER
 							[
-								SNew(SVerticalBox)
-								+ SVerticalBox::Slot()
+								SNew(SHorizontalBox)
+								+ SHorizontalBox::Slot()
 								.Padding(4)
-								.AutoHeight()
 								[
-									NGED_SECTION_DARKBORDER
+									SNew(SVerticalBox)
+									+ SVerticalBox::Slot()
+									.Padding(2)
+									.AutoHeight()
 									[
-										SNew(SHorizontalBox)
-										+ SHorizontalBox::Slot()
-										.HAlign(HAlign_Right)
-										.VAlign(VAlign_Center)
-										.AutoWidth()
-										.Padding(4)
+										NGED_SECTION_DARKBORDER
 										[
-											SNew(STextBlock)
-											.Text(FString("Update Script"))
+											SNew(SHorizontalBox)
+											+ SHorizontalBox::Slot()
+											.HAlign(HAlign_Right)
+											.VAlign(VAlign_Center)
+											.AutoWidth()
+											.Padding(2)
+											[
+												SNew(STextBlock)
+												.Text(FString("Update Script"))
+											]
+											+ SHorizontalBox::Slot()
+											.HAlign(HAlign_Right)
+											.VAlign(VAlign_Center)
+											[
+												SAssignNew(UpdateScriptSelector, SContentReference)
+												.AllowedClass(UNiagaraScript::StaticClass())
+												.AssetReference(this, &SEmitterWidget::GetUpdateScript)
+												.AllowSelectingNewAsset(true)
+												.AllowClearingReference(true)
+												.OnSetReference(this, &SEmitterWidget::OnUpdateScriptSelectedFromPicker)
+												.WidthOverride(130)
+											]
 										]
-										+ SHorizontalBox::Slot()
-										.HAlign(HAlign_Right)
-										.VAlign(VAlign_Center)
+									]
+
+									+ SVerticalBox::Slot()
+									.Padding(2)
+									.AutoHeight()
+									.Expose(UpdateScriptConstantListSlot)
+									[
+										SNew(SExpandableArea).InitiallyCollapsed(true).AreaTitle(FString("Constants"))
+										.BodyContent()
 										[
-											SAssignNew(UpdateScriptSelector, SContentReference)
-											.AllowedClass(UNiagaraScript::StaticClass())
-											.AssetReference(this, &SEmitterWidget::GetUpdateScript)
-											.AllowSelectingNewAsset(true)
-											.AllowClearingReference(true)
-											.OnSetReference(this, &SEmitterWidget::OnUpdateScriptSelectedFromPicker)
-											.WidthOverride(150)
+											SAssignNew(UpdateScriptConstantList, SListView<TSharedPtr<EditorExposedVectorConstant>>)
+											.ItemHeight(20).ListItemsSource(EditorExposedConstants).OnGenerateRow(this, &SEmitterWidget::OnGenerateConstantListRow)
+											.SelectionMode(ESelectionMode::None)
 										]
 									]
 								]
 
-								+ SVerticalBox::Slot()
+								+ SHorizontalBox::Slot()
 								.Padding(4)
-								.AutoHeight()
-								.Expose(UpdateScriptConstantListSlot)
 								[
-									SAssignNew(UpdateScriptConstantList, SListView<TSharedPtr<EditorExposedVectorConstant>>)
-									.ItemHeight(20).ListItemsSource(EditorExposedConstants).OnGenerateRow(this, &SEmitterWidget::OnGenerateConstantListRow)
-									.SelectionMode(ESelectionMode::None)
-								]
-
-
-
-
-								+ SVerticalBox::Slot()
-								.Padding(4)
-								.AutoHeight()
-								[
-									NGED_SECTION_DARKBORDER
+									SNew(SVerticalBox)
+									+ SVerticalBox::Slot()
+									.Padding(2)
+									.AutoHeight()
 									[
-										SNew(SHorizontalBox)
-										+ SHorizontalBox::Slot()
-										.HAlign(HAlign_Right)
-										.VAlign(VAlign_Center)
-										.AutoWidth()
-										.Padding(4)
+										NGED_SECTION_DARKBORDER
 										[
-											SNew(STextBlock)
-											.Text(FString("Spawn Script"))
+											SNew(SHorizontalBox)
+											+ SHorizontalBox::Slot()
+											.HAlign(HAlign_Right)
+											.VAlign(VAlign_Center)
+											.AutoWidth()
+											.Padding(4)
+											[
+												SNew(STextBlock)
+												.Text(FString("Spawn Script"))
+											]
+											+ SHorizontalBox::Slot()
+											.HAlign(HAlign_Right)
+											.VAlign(VAlign_Center)
+											[
+												SAssignNew(SpawnScriptSelector, SContentReference)
+												.AllowedClass(UNiagaraScript::StaticClass())
+												.AssetReference(this, &SEmitterWidget::GetSpawnScript)
+												.AllowSelectingNewAsset(true)
+												.AllowClearingReference(true)
+												.OnSetReference(this, &SEmitterWidget::OnSpawnScriptSelectedFromPicker)
+												.WidthOverride(130)
+											]
 										]
-										+ SHorizontalBox::Slot()
-										.HAlign(HAlign_Right)
-										.VAlign(VAlign_Center)
+									]
+
+									+ SVerticalBox::Slot()
+										.Padding(2)
+										.AutoHeight()
+										.Expose(SpawnScriptConstantListSlot)
 										[
-											SAssignNew(SpawnScriptSelector, SContentReference)
-											.AllowedClass(UNiagaraScript::StaticClass())
-											.AssetReference(this, &SEmitterWidget::GetSpawnScript)
-											.AllowSelectingNewAsset(true)
-											.AllowClearingReference(true)
-											.OnSetReference(this, &SEmitterWidget::OnSpawnScriptSelectedFromPicker)
-											.WidthOverride(150)
+											SNew(SExpandableArea).InitiallyCollapsed(true).AreaTitle(FString("Constants"))
+											.BodyContent()
+											[
+												SAssignNew(SpawnScriptConstantList, SListView<TSharedPtr<EditorExposedVectorConstant>>)
+												.ItemHeight(20).ListItemsSource(EditorExposedSpawnConstants).OnGenerateRow(this, &SEmitterWidget::OnGenerateConstantListRow)
+												.SelectionMode(ESelectionMode::None)
+											]
 										]
-									]
 								]
-
-								+ SVerticalBox::Slot()
-									.Padding(4)
-									.AutoHeight()
-									.Expose(SpawnScriptConstantListSlot)
-									[
-										SAssignNew(SpawnScriptConstantList, SListView<TSharedPtr<EditorExposedVectorConstant>>)
-										.ItemHeight(20).ListItemsSource(EditorExposedSpawnConstants).OnGenerateRow(this, &SEmitterWidget::OnGenerateConstantListRow)
-										.SelectionMode(ESelectionMode::None)
-									]
 							]
 						]
 
@@ -370,11 +401,13 @@ void SEmitterWidget::Construct(const FArguments& InArgs)
 								[
 									SNew(SVerticalBox)
 									+ SVerticalBox::Slot()
+									.VAlign(VAlign_Center)
 									.AutoHeight()
 									.Padding(4)
 									[
 										SNew(SHorizontalBox)
 										+ SHorizontalBox::Slot()
+										.VAlign(VAlign_Center)
 										.Padding(4)
 										.AutoWidth()
 										[
@@ -417,10 +450,12 @@ void SEmitterWidget::Construct(const FArguments& InArgs)
 									]
 									+ SVerticalBox::Slot()
 										.AutoHeight()
+										.VAlign(VAlign_Center)
 										.Padding(4)
 										[
 											SNew(SHorizontalBox)
 											+ SHorizontalBox::Slot()
+											.VAlign(VAlign_Center)
 											.Padding(4)
 											.AutoWidth()
 											[
@@ -428,26 +463,110 @@ void SEmitterWidget::Construct(const FArguments& InArgs)
 												.Text(FString("Render as"))
 											]
 											+ SHorizontalBox::Slot()
-												.Padding(4)
-												.AutoWidth()
+											.Padding(4)
+											.VAlign(VAlign_Center)
+											.AutoWidth()
+											[
+												// session combo box
+												SNew(SComboBox<TSharedPtr<FString> >)
+												.ContentPadding(FMargin(6.0f, 2.0f))
+												.OptionsSource(&RenderModuleList)
+												.OnSelectionChanged(this, &SEmitterWidget::OnRenderModuleChanged)
+												//										.ToolTipText(LOCTEXT("SessionComboBoxToolTip", "Select the game session to interact with.").ToString())
+												.OnGenerateWidget(this, &SEmitterWidget::GenerateRenderModuleComboWidget)
 												[
-													// session combo box
-													SNew(SComboBox<TSharedPtr<FString> >)
-													.ContentPadding(FMargin(6.0f, 2.0f))
-													.OptionsSource(&RenderModuleList)
-													.OnSelectionChanged(this, &SEmitterWidget::OnRenderModuleChanged)
-													//										.ToolTipText(LOCTEXT("SessionComboBoxToolTip", "Select the game session to interact with.").ToString())
-													.OnGenerateWidget(this, &SEmitterWidget::GenerateRenderModuleComboWidget)
-													[
-														SNew(STextBlock)
-														.Text(this, &SEmitterWidget::GetRenderModuleText)
-													]
+													SNew(STextBlock)
+													.Text(this, &SEmitterWidget::GetRenderModuleText)
 												]
+											]
 										]
+									+ SVerticalBox::Slot()
+										.AutoHeight()
+										.Padding(4)
+										[
+											RendererPropertiesViewRef
+										]
+
 								]
 							]
 					]
 			]
 		];
 
+}
+
+
+
+
+void SEmitterWidget::OnRenderModuleChanged(TSharedPtr<FString> ModName, ESelectInfo::Type SelectInfo)
+{
+	EEmitterRenderModuleType Type = RMT_Sprites;
+	for (int32 i = 0; i < RenderModuleList.Num(); i++)
+	{
+		if (*RenderModuleList[i] == *ModName)
+		{
+			Type = (EEmitterRenderModuleType)(i + 1);
+			break;
+		}
+	}
+
+
+	ENQUEUE_UNIQUE_RENDER_COMMAND_THREEPARAMETER(
+		FChangeNiagaraRenderModule,
+		FNiagaraEffectInstance*, InEffect, this->EffectInstance,
+		TSharedPtr<FNiagaraSimulation>, InEmitter, this->Emitter,
+		EEmitterRenderModuleType, InType, Type,
+		{
+		InEmitter->SetRenderModuleType(InType, InEffect->GetComponent()->GetWorld()->FeatureLevel);
+		InEmitter->GetProperties()->RenderModuleType = InType;
+		InEffect->RenderModuleupdate();
+	});
+	TComponentReregisterContext<UNiagaraComponent> ComponentReregisterContext;
+
+	UObject *Props = Emitter->GetProperties()->RendererProperties;
+	RendererPropertiesView->SetObject(Props);
+}
+
+
+
+
+void SEmitterWidget::OnUpdateScriptSelectedFromPicker(UObject *Asset)
+{
+	// Close the asset picker
+	//AssetPickerAnchor->SetIsOpen(false);
+
+	// Set the object found from the asset picker
+	CurUpdateScript = Cast<UNiagaraScript>(Asset);
+	Emitter->GetProperties()->UpdateScript = CurUpdateScript;
+
+	UpdateScriptConstantListSlot->DetachWidget();
+	(*UpdateScriptConstantListSlot)
+		[
+			SNew(SExpandableArea).InitiallyCollapsed(true).AreaTitle(FString("Constants"))
+			.BodyContent()
+			[
+				SAssignNew(UpdateScriptConstantList, SListView<TSharedPtr<EditorExposedVectorConstant>>)
+				.ItemHeight(20).ListItemsSource(&CurUpdateScript->Source->ExposedVectorConstants).OnGenerateRow(this, &SEmitterWidget::OnGenerateConstantListRow)
+			]
+		];
+}
+
+void SEmitterWidget::OnSpawnScriptSelectedFromPicker(UObject *Asset)
+{
+	// Close the asset picker
+	//AssetPickerAnchor->SetIsOpen(false);
+
+	// Set the object found from the asset picker
+	CurSpawnScript = Cast<UNiagaraScript>(Asset);
+	Emitter->GetProperties()->SpawnScript = CurSpawnScript;
+	SpawnScriptConstantListSlot->DetachWidget();
+	(*SpawnScriptConstantListSlot)
+		[
+			SNew(SExpandableArea).InitiallyCollapsed(true).AreaTitle(FString("Constants"))
+			.BodyContent()
+			[
+				SAssignNew(SpawnScriptConstantList, SListView<TSharedPtr<EditorExposedVectorConstant>>)
+				.ItemHeight(20).ListItemsSource(&CurSpawnScript->Source->ExposedVectorConstants).OnGenerateRow(this, &SEmitterWidget::OnGenerateConstantListRow)
+			]
+		];
 }
\ No newline at end of file
diff --git a/Engine/Source/Editor/NiagaraEditor/Private/SNiagaraEffectEditorWidget.h b/Engine/Source/Editor/NiagaraEditor/Private/SNiagaraEffectEditorWidget.h
--- a/Engine/Source/Editor/NiagaraEditor/Private/SNiagaraEffectEditorWidget.h
+++ b/Engine/Source/Editor/NiagaraEditor/Private/SNiagaraEffectEditorWidget.h
@@ -15,11 +15,16 @@
 #include "Engine/NiagaraSimulation.h"
 #include "Engine/NiagaraEffectRenderer.h"
 #include "ComponentReregisterContext.h"
+#include "SNumericEntryBox.h"
 
-#define NGED_SECTION_BORDER SNew(SBorder).BorderImage(FEditorStyle::GetBrush("ToolPanel.GroupBorder")).Padding(4.0f).HAlign(HAlign_Left)
-#define NGED_SECTION_LIGHTBORDER SNew(SBorder).BorderImage(FEditorStyle::GetBrush("ToolPanel.LightGroupBorder")).Padding(4.0f).HAlign(HAlign_Left)
-#define NGED_SECTION_DARKBORDER SNew(SBorder).BorderImage(FEditorStyle::GetBrush("ToolPanel.DarkGroupBorder")).Padding(4.0f).HAlign(HAlign_Left)
+#define NGED_SECTION_BORDER SNew(SBorder).BorderImage(FEditorStyle::GetBrush("ToolPanel.GroupBorder")).Padding(2.0f).HAlign(HAlign_Left)
+#define NGED_SECTION_LIGHTBORDER SNew(SBorder).BorderImage(FEditorStyle::GetBrush("ToolPanel.LightGroupBorder")).Padding(2.0f).HAlign(HAlign_Left)
+#define NGED_SECTION_DARKBORDER SNew(SBorder).BorderImage(FEditorStyle::GetBrush("ToolPanel.DarkGroupBorder")).Padding(2.0f).HAlign(HAlign_Left)
 
+const FLinearColor RedCol(0.5, 0.1, 0.1);
+const FLinearColor GreenCol(0.1, 0.5, 0.1);
+const FLinearColor BlueCol(0.1, 0.1, 0.5);
+const FLinearColor GrayCol(0.3, 0.3, 0.3);
 
 class SVectorConstantWidget : public SCompoundWidget, public FNotifyHook
 {
@@ -27,7 +32,7 @@ private:
 	FName ConstantName;
 	TSharedPtr<FNiagaraSimulation> Emitter;
 
-	TSharedPtr<SEditableTextBox> XText, YText, ZText, WText;
+	//SNumericEntryBox<float> XText, YText, ZText, WText;
 
 public:
 	SLATE_BEGIN_ARGS(SVectorConstantWidget) :
@@ -39,27 +44,101 @@ public:
 	SLATE_ARGUMENT(FName, ConstantName)
 	SLATE_END_ARGS()
 
+	TOptional<float> GetConstX() const	
+	{
+		FVector4 *ConstPtr = Emitter->GetProperties()->ExternalConstants.FindVector(ConstantName);
+		if (ConstPtr)
+		{
+			return ConstPtr->X;
+		}
+		return 0.0f;
+	}
 
-	void OnConstantChanged(const FText &InText)
+	TOptional<float> GetConstY() const
 	{
-		FVector4 Vec;
-		FString Strns[4];
-		Strns[0] = XText->GetText().ToString();
-		Strns[1] = YText->GetText().ToString();
-		Strns[2] = ZText->GetText().ToString();
-		Strns[3] = WText->GetText().ToString();
-		for (uint32 i = 0; i < 4; i++)
+		FVector4 *ConstPtr = Emitter->GetProperties()->ExternalConstants.FindVector(ConstantName);
+		if (ConstPtr)
 		{
-			if (Strns[i]!="")
-			{
-				Vec[i] = FCString::Atof(*Strns[i]);
-			}
-			else
-			{
-				Vec[i] = 0.0f;
-			}
+			return ConstPtr->Y;
 		}
-		Emitter->GetProperties()->ExternalConstants.SetOrAdd(ConstantName, Vec);
+		return 0.0f;
+	}
+
+	TOptional<float> GetConstZ() const
+	{
+		FVector4 *ConstPtr = Emitter->GetProperties()->ExternalConstants.FindVector(ConstantName);
+		if (ConstPtr)
+		{
+			return ConstPtr->Z;
+		}
+		return 0.0f;
+	}
+
+	TOptional<float> GetConstW() const
+	{
+		FVector4 *ConstPtr = Emitter->GetProperties()->ExternalConstants.FindVector(ConstantName);
+		if (ConstPtr)
+		{
+			return ConstPtr->W;
+		}
+		return 0.0f;
+	}
+
+	void OnConstantChangedX(float InVal)	{ OnConstantChangedX(InVal, ETextCommit::Default); }
+	void OnConstantChangedY(float InVal)	{ OnConstantChangedY(InVal, ETextCommit::Default); }
+	void OnConstantChangedZ(float InVal)	{ OnConstantChangedZ(InVal, ETextCommit::Default); }
+	void OnConstantChangedW(float InVal)	{ OnConstantChangedW(InVal, ETextCommit::Default); }
+
+	void OnConstantChangedX(float InVal, ETextCommit::Type Type)
+	{
+		FVector4 Constant(InVal, GetConstY().GetValue(), GetConstZ().GetValue(), GetConstW().GetValue());
+		Emitter->GetProperties()->ExternalConstants.SetOrAdd(ConstantName, Constant);
+		/*
+		FVector4 *VecPtr = Emitter->GetProperties()->ExternalConstants.FindVector(ConstantName);
+		if (VecPtr)
+		{
+			VecPtr->X = InVal;
+		}
+		*/
+	}
+
+	void OnConstantChangedY(float InVal, ETextCommit::Type Type)
+	{
+		FVector4 Constant(GetConstX().GetValue(), InVal, GetConstZ().GetValue(), GetConstW().GetValue());
+		Emitter->GetProperties()->ExternalConstants.SetOrAdd(ConstantName, Constant);
+		/*
+		FVector4 *VecPtr = Emitter->GetProperties()->ExternalConstants.FindVector(ConstantName);
+		if (VecPtr)
+		{
+			VecPtr->Y = InVal;
+		}
+		*/
+	}
+
+	void OnConstantChangedZ(float InVal, ETextCommit::Type Type)
+	{
+		FVector4 Constant(GetConstX().GetValue(), GetConstY().GetValue(), InVal, GetConstW().GetValue());
+		Emitter->GetProperties()->ExternalConstants.SetOrAdd(ConstantName, Constant);
+		/*
+		FVector4 *VecPtr = Emitter->GetProperties()->ExternalConstants.FindVector(ConstantName);
+		if (VecPtr)
+		{
+			VecPtr->Z = InVal;
+		}
+		*/
+	}
+
+	void OnConstantChangedW(float InVal, ETextCommit::Type Type)
+	{
+		FVector4 Constant(GetConstX().GetValue(), GetConstY().GetValue(), GetConstZ().GetValue(), InVal);
+		Emitter->GetProperties()->ExternalConstants.SetOrAdd(ConstantName, Constant);
+		/*
+		FVector4 *VecPtr = Emitter->GetProperties()->ExternalConstants.FindVector(ConstantName);
+		if (VecPtr)
+		{
+			VecPtr->W = InVal;
+		}
+		*/
 	}
 
 	void Construct(const FArguments& InArgs)
@@ -67,24 +146,48 @@ public:
 		Emitter = InArgs._Emitter;
 		ConstantName = InArgs._ConstantName;
 
-		FVector4 *ConstVal = Emitter->GetProperties()->ExternalConstants.FindVector(ConstantName);
-		FString ConstText[4] = {"","","",""};
-		if (ConstVal)
+		FVector4 *ConstPtr = Emitter->GetProperties()->ExternalConstants.FindVector(ConstantName);
+		FVector4 ConstVal;
+		if (ConstPtr)
 		{
-			for (uint32 i = 0; i < 4; i++)
-			{
-				ConstText[i] = FString::Printf(TEXT("%.4f"), (*ConstVal)[i]);
-			}
+			ConstVal = *ConstPtr;
 		}
 		
 		this->ChildSlot
 		[
 		SNew(SHorizontalBox)
-		+ SHorizontalBox::Slot().Padding(4)[SNew(STextBlock).Text(FText::FromName(ConstantName))]
-		+ SHorizontalBox::Slot().Padding(4)[SAssignNew(XText, SEditableTextBox).Text(FText::FromString(ConstText[0])).OnTextChanged(this, &SVectorConstantWidget::OnConstantChanged)]
-		+ SHorizontalBox::Slot().Padding(4)[SAssignNew(YText, SEditableTextBox).Text(FText::FromString(ConstText[1])).OnTextChanged(this, &SVectorConstantWidget::OnConstantChanged)]
-		+ SHorizontalBox::Slot().Padding(4)[SAssignNew(ZText, SEditableTextBox).Text(FText::FromString(ConstText[2])).OnTextChanged(this, &SVectorConstantWidget::OnConstantChanged)]
-		+ SHorizontalBox::Slot().Padding(4)[SAssignNew(WText, SEditableTextBox).Text(FText::FromString(ConstText[3])).OnTextChanged(this, &SVectorConstantWidget::OnConstantChanged)]
+		+ SHorizontalBox::Slot().VAlign(VAlign_Center).Padding(4)[SNew(STextBlock).Text(FText::FromName(ConstantName))]
+		+ SHorizontalBox::Slot().VAlign(VAlign_Center).Padding(4)[SNew(SNumericEntryBox<float>).Value(this, &SVectorConstantWidget::GetConstX)
+		.OnValueCommitted(this, &SVectorConstantWidget::OnConstantChangedX).OnValueChanged(this, &SVectorConstantWidget::OnConstantChangedX)
+			.Label()
+			[
+				SNumericEntryBox<float>::BuildLabel(FText::FromString("X"), FLinearColor::White, RedCol)
+			]
+		]
+
+			+ SHorizontalBox::Slot().VAlign(VAlign_Center).Padding(4)[SNew(SNumericEntryBox<float>).Value(this, &SVectorConstantWidget::GetConstY)
+			.OnValueCommitted(this, &SVectorConstantWidget::OnConstantChangedY).OnValueChanged(this, &SVectorConstantWidget::OnConstantChangedY)
+			.Label()
+			[
+				SNumericEntryBox<float>::BuildLabel(FText::FromString("Y"), FLinearColor::White, GreenCol)
+			]
+		]
+
+			+ SHorizontalBox::Slot().VAlign(VAlign_Center).Padding(4)[SNew(SNumericEntryBox<float>).Value(this, &SVectorConstantWidget::GetConstZ)
+			.OnValueCommitted(this, &SVectorConstantWidget::OnConstantChangedZ).OnValueChanged(this, &SVectorConstantWidget::OnConstantChangedZ)
+			.Label()
+			[
+				SNumericEntryBox<float>::BuildLabel(FText::FromString("Z"), FLinearColor::White, BlueCol)
+			]
+		]
+
+			+ SHorizontalBox::Slot().VAlign(VAlign_Center).Padding(4)[SNew(SNumericEntryBox<float>).Value(this, &SVectorConstantWidget::GetConstW)
+			.OnValueCommitted(this, &SVectorConstantWidget::OnConstantChangedW).OnValueChanged(this, &SVectorConstantWidget::OnConstantChangedW)
+			.Label()
+			[
+				SNumericEntryBox<float>::BuildLabel(FText::FromString("W"), FLinearColor::White, GrayCol)
+			]
+		]
 		];
 		
 	}
@@ -105,6 +208,8 @@ private:
 	UNiagaraScript *CurUpdateScript, *CurSpawnScript;
 	UMaterial *CurMaterial;
 
+	TSharedPtr<IDetailsView> RendererPropertiesView;
+
 	TArray<TSharedPtr<FString> > RenderModuleList;
 
 	TSharedPtr< SListView<TSharedPtr<EditorExposedVectorConstant>>> UpdateScriptConstantList;
@@ -131,33 +236,19 @@ public:
 	}
 
 
-	void OnUpdateScriptSelectedFromPicker(UObject *Asset)
+	void OnEmitterNameChanged(const FText &InText)
 	{
-		// Close the asset picker
-		//AssetPickerAnchor->SetIsOpen(false);
-
-		// Set the object found from the asset picker
-		CurUpdateScript = Cast<UNiagaraScript>(Asset);
-		Emitter->GetProperties()->UpdateScript = CurUpdateScript;
-
-		UpdateScriptConstantListSlot->DetachWidget();
-		(*UpdateScriptConstantListSlot)
-			[
-				SAssignNew(UpdateScriptConstantList, SListView<TSharedPtr<EditorExposedVectorConstant>>)
-				.ItemHeight(20).ListItemsSource(&CurUpdateScript->Source->ExposedVectorConstants).OnGenerateRow(this, &SEmitterWidget::OnGenerateConstantListRow)
-			];
+		Emitter->GetProperties()->Name = InText.ToString();
 	}
 
-	void OnSpawnScriptSelectedFromPicker(UObject *Asset)
+	FText GetEmitterName() const
 	{
-		// Close the asset picker
-		//AssetPickerAnchor->SetIsOpen(false);
-
-		// Set the object found from the asset picker
-		CurSpawnScript = Cast<UNiagaraScript>(Asset);
-		Emitter->GetProperties()->SpawnScript = CurSpawnScript;
+		return FText::FromString(Emitter->GetProperties()->Name);
 	}
 
+	void OnUpdateScriptSelectedFromPicker(UObject *Asset);
+	void OnSpawnScriptSelectedFromPicker(UObject *Asset);
+
 	void OnMaterialSelected(UObject *Asset)
 	{
 		MaterialThumbnail->SetAsset(Asset);
@@ -167,6 +258,7 @@ public:
 		CurMaterial = Cast<UMaterial>(Asset);
 		Emitter->GetEffectRenderer()->SetMaterial(CurMaterial, ERHIFeatureLevel::SM5);
 		Emitter->GetProperties()->Material = CurMaterial;
+		TComponentReregisterContext<UNiagaraComponent> ComponentReregisterContext;
 	}
 
 	UObject *GetMaterial() const
@@ -195,31 +287,8 @@ public:
 		return Emitter->IsEnabled() ? ESlateCheckBoxState::Checked : ESlateCheckBoxState::Unchecked;
 	}
 
-	void OnRenderModuleChanged(TSharedPtr<FString> ModName, ESelectInfo::Type SelectInfo)
-	{
-		EEmitterRenderModuleType Type = RMT_Sprites;
-		for (int32 i = 0; i < RenderModuleList.Num(); i++)
-		{
-			if (*RenderModuleList[i] == *ModName)
-			{
-				Type = (EEmitterRenderModuleType)(i+1);
-				break;
-			}
-		}
-
+	void OnRenderModuleChanged(TSharedPtr<FString> ModName, ESelectInfo::Type SelectInfo);
 
-		ENQUEUE_UNIQUE_RENDER_COMMAND_THREEPARAMETER(
-			FChangeNiagaraRenderModule,
-			FNiagaraEffectInstance*, InEffect, this->EffectInstance,
-			TSharedPtr<FNiagaraSimulation>, InEmitter, this->Emitter,
-			EEmitterRenderModuleType, InType, Type,
-		{
-			InEmitter->SetRenderModuleType(InType, InEffect->GetComponent()->GetWorld()->FeatureLevel);
-			InEmitter->GetProperties()->RenderModuleType = InType;
-			InEffect->RenderModuleupdate();
-		});
-		TComponentReregisterContext<UNiagaraComponent> ComponentReregisterContext;
-	}
 
 	void OnSpawnRateChanged(const FText &InText)
 	{
@@ -256,7 +325,7 @@ public:
 	FString GetStatsText() const
 	{
 		TCHAR Txt[128];
-		FCString::Snprintf(Txt, 128, TEXT("%i Particles"), Emitter->GetNumParticles());
+		FCString::Snprintf(Txt, 128, TEXT("%i Particles  |  %.2f ms  |  %.1f MB"), Emitter->GetNumParticles(), Emitter->GetTotalCPUTime(), Emitter->GetTotalBytesUsed()/(1024.0*1024.0));
 		return FString(Txt);
 	}
 
diff --git a/Engine/Source/Editor/UnrealEd/Private/Niagara/NiagaraEditorCommon.cpp b/Engine/Source/Editor/UnrealEd/Private/Niagara/NiagaraEditorCommon.cpp
--- a/Engine/Source/Editor/UnrealEd/Private/Niagara/NiagaraEditorCommon.cpp
+++ b/Engine/Source/Editor/UnrealEd/Private/Niagara/NiagaraEditorCommon.cpp
@@ -509,7 +509,17 @@ void FNiagaraOpInfo::Init()
 	Op->Outputs.Add(FNiagaraOpInOutInfo(Result, ENiagaraDataType::Matrix, ResultText, ResultText, DefaultStr_VecOne));
 	Op->OpDelegate.BindStatic(&INiagaraCompiler::Inverse);
 
-#define NiagaraOp(OpName) if (GetOpInfo(OpName) == NULL)\
+
+	Op = &OpInfos.Add(LessThan);
+	Op->Name = LessThan;
+	Op->FriendlyName = NSLOCTEXT("NiagaraOpInfo", "LessThan Name", "A Less Than B");
+	Op->Description = NSLOCTEXT("NiagaraOpInfo", "LessThan Desc", "Result = 1 if A<B, 0 otherwise");
+	Op->Inputs.Add(FNiagaraOpInOutInfo(A, ENiagaraDataType::Vector, AText, AText, DefaultStr_VecOne));
+	Op->Inputs.Add(FNiagaraOpInOutInfo(B, ENiagaraDataType::Vector, BText, BText, DefaultStr_VecOne));
+	Op->Outputs.Add(FNiagaraOpInOutInfo(Result, ENiagaraDataType::Vector, ResultText, ResultText, DefaultStr_VecOne));
+	Op->OpDelegate.BindStatic(&INiagaraCompiler::LessThan);
+
+	#define NiagaraOp(OpName) if (GetOpInfo(OpName) == NULL)\
 	{\
 	UE_LOG(LogNiagaraCompiler, Fatal, TEXT("Couldn't find info for \"%s\". Have you added it to the NiagaraOpList but not provided the extra information in FNiagaraOpInfo::Init() ?"), TEXT(#OpName)); \
 	}
diff --git a/Engine/Source/Editor/UnrealEd/Public/NiagaraEditorCommon.h b/Engine/Source/Editor/UnrealEd/Public/NiagaraEditorCommon.h
--- a/Engine/Source/Editor/UnrealEd/Public/NiagaraEditorCommon.h
+++ b/Engine/Source/Editor/UnrealEd/Public/NiagaraEditorCommon.h
@@ -60,7 +60,8 @@
 	NiagaraOp(GetW)\
 	NiagaraOp(TransformVector)\
 	NiagaraOp(Transpose)\
-	NiagaraOp(Inverse)
+	NiagaraOp(Inverse)\
+	NiagaraOp(LessThan)
 
 enum class ENiagaraExpressionResultLocation
 {
diff --git a/Engine/Source/Runtime/Engine/Classes/Components/NiagaraComponent.h b/Engine/Source/Runtime/Engine/Classes/Components/NiagaraComponent.h
--- a/Engine/Source/Runtime/Engine/Classes/Components/NiagaraComponent.h
+++ b/Engine/Source/Runtime/Engine/Classes/Components/NiagaraComponent.h
@@ -71,6 +71,7 @@ public:
 	FVector4 *GetCurrentBuffer()		{ return ParticleBuffers[CurrentBuffer].GetData(); }
 	FVector4 *GetPreviousBuffer()		{ return ParticleBuffers[CurrentBuffer^0x1].GetData(); }
 
+	int GetBytesUsed()	{ return (ParticleBuffers[0].Num() + ParticleBuffers[1].Num()) * 16 + Attributes.Num() * 4; }
 private:
 	uint32 CurrentBuffer, NumParticles, ParticleAllocation;
 	TArray<FVector4> ParticleBuffers[2];
diff --git a/Engine/Source/Runtime/Engine/Classes/Engine/NiagaraEffect.h b/Engine/Source/Runtime/Engine/Classes/Engine/NiagaraEffect.h
--- a/Engine/Source/Runtime/Engine/Classes/Engine/NiagaraEffect.h
+++ b/Engine/Source/Runtime/Engine/Classes/Engine/NiagaraEffect.h
@@ -3,6 +3,7 @@
 #pragma once
 #include "NiagaraSimulation.h"
 #include "NiagaraEffectRenderer.h"
+#include "List.h"
 #include "NiagaraEffect.generated.h"
 
 UCLASS(MinimalAPI)
@@ -11,29 +12,42 @@ class UNiagaraEffect : public UObject
 	GENERATED_UCLASS_BODY()
 
 public:
-	ENGINE_API FNiagaraEmitterProperties &AddEmitterProperties();
+	ENGINE_API FNiagaraEmitterProperties *AddEmitterProperties();
 
-	void Init(UNiagaraComponent *InComponent)
+	void Init()
 	{
-		Component = InComponent;
 	}
 
 
-	TArray<FNiagaraEmitterProperties> &GetEmitterProperties() {
-		return EmitterProps;
+	FNiagaraEmitterProperties *GetEmitterProperties(int Idx)
+	{
+		check(Idx < EmitterProps.Num());
+		return EmitterProps[Idx];
 	};
 
+	int GetNumEmitters()
+	{
+		return EmitterProps.Num();
+	}
+
+	ENGINE_API void CreateEffectRendererProps(TSharedPtr<FNiagaraSimulation> Sim);
+
 	virtual void PostLoad() override; 
+	virtual void PreSave() override;
 
 private:
-	UNiagaraComponent *Component;
-
+	// serialized array of emitter properties
 	UPROPERTY()
-	TArray<FNiagaraEmitterProperties>EmitterProps;
+	TArray<FNiagaraEmitterProperties>EmitterPropsSerialized;
+
+	TArray<FNiagaraEmitterProperties*> EmitterProps;
 };
 
 
+class FNiagaraEmitterProp
+{
 
+};
 
 
 
@@ -43,9 +57,9 @@ public:
 	explicit FNiagaraEffectInstance(UNiagaraEffect *InAsset, UNiagaraComponent *InComponent)
 	{
 		Component = InComponent;
-		for (int32 i = 0; i < InAsset->GetEmitterProperties().Num(); i++)
+		for (int32 i = 0; i < InAsset->GetNumEmitters(); i++)
 		{
-			FNiagaraSimulation *Sim = new FNiagaraSimulation(&InAsset->GetEmitterProperties()[i], Component->GetWorld()->FeatureLevel);
+			FNiagaraSimulation *Sim = new FNiagaraSimulation(InAsset->GetEmitterProperties(i), Component->GetWorld()->FeatureLevel);
 			Emitters.Add(MakeShareable(Sim));
 		}
 		InitRenderModules(Component->GetWorld()->FeatureLevel);
@@ -54,30 +68,45 @@ public:
 	explicit FNiagaraEffectInstance(UNiagaraEffect *InAsset)
 		: Component(nullptr)
 	{
-		for (FNiagaraEmitterProperties &Props : InAsset->GetEmitterProperties())
+		InitEmitters(InAsset);
+	}
+
+	void InitEmitters(UNiagaraEffect *InAsset)
+	{
+		check(InAsset);
+		for (int i = 0; i < InAsset->GetNumEmitters(); i++)
 		{
-			FNiagaraSimulation *Sim = new FNiagaraSimulation(&Props);
+			FNiagaraEmitterProperties *Props = InAsset->GetEmitterProperties(i);
+			FNiagaraSimulation *Sim = new FNiagaraSimulation(Props);
 			Emitters.Add(MakeShareable(Sim));
 		}
 	}
 
-
 	void InitRenderModules(ERHIFeatureLevel::Type InFeatureLevel)
 	{
 		for (TSharedPtr<FNiagaraSimulation> Sim : Emitters)
 		{
+			if(Sim->GetProperties()->RendererProperties == nullptr)
+			{
+				Component->GetAsset()->CreateEffectRendererProps(Sim);
+			}
 			Sim->SetRenderModuleType(Sim->GetProperties()->RenderModuleType, InFeatureLevel);
 		}
 	}
 
 	void Init(UNiagaraComponent *InComponent)
 	{
+		if (InComponent->GetAsset())
+		{
+			Emitters.Empty();
+			InitEmitters(InComponent->GetAsset());
+		}
 		Component = InComponent;
-		//InitRenderModules(Component->GetWorld()->FeatureLevel);
+		InitRenderModules(Component->GetWorld()->FeatureLevel);
+		RenderModuleupdate();
 	}
 
-	ENGINE_API void AddEmitter(FNiagaraEmitterProperties *Properties);
-
+	ENGINE_API TSharedPtr<FNiagaraSimulation> AddEmitter(FNiagaraEmitterProperties *Properties);
 
 	void SetConstant(FName ConstantName, const float Value)
 	{
@@ -110,7 +139,10 @@ public:
 	void RenderModuleupdate()
 	{
 		FNiagaraSceneProxy *NiagaraProxy = static_cast<FNiagaraSceneProxy*>(Component->SceneProxy);
-		NiagaraProxy->UpdateEffectRenderers(this);
+		if (NiagaraProxy)
+		{
+			NiagaraProxy->UpdateEffectRenderers(this);
+		}
 	}
 
 	FNiagaraSimulation *GetEmitter(uint32 idx)
diff --git a/Engine/Source/Runtime/Engine/Classes/Engine/NiagaraEffectRenderer.h b/Engine/Source/Runtime/Engine/Classes/Engine/NiagaraEffectRenderer.h
--- a/Engine/Source/Runtime/Engine/Classes/Engine/NiagaraEffectRenderer.h
+++ b/Engine/Source/Runtime/Engine/Classes/Engine/NiagaraEffectRenderer.h
@@ -10,6 +10,9 @@ NiagaraEffectRenderer.h: Base class for Niagara render modules
 #include "ParticleVertexFactory.h"
 #include "ParticleBeamTrailVertexFactory.h"
 #include "Components/NiagaraComponent.h"
+#include "NiagaraEffectRendererProperties.h"
+#include "NiagaraSpriteRendererProperties.h"
+#include "NiagaraRibbonRendererProperties.h"
 
 DECLARE_CYCLE_STAT(TEXT("Generate Sprite Vertex Data"), STAT_NiagaraGenSpriteVertexData, STATGROUP_Niagara);
 DECLARE_CYCLE_STAT(TEXT("Generate Ribbon Vertex Data"), STAT_NiagaraGenRibbonVertexData, STATGROUP_Niagara);
@@ -22,8 +25,6 @@ struct FNiagaraDynamicDataBase
 
 
 
-
-
 /**
 * Base class for Niagara effect renderers. Effect renderers handle generating vertex data for and
 * drawing of simulation data coming out of FNiagaraSimulation instances.
@@ -73,7 +74,13 @@ public:
 
 	virtual ~NiagaraEffectRenderer() {}
 
+	virtual UClass *GetPropertiesClass() = 0;
+
+	float GetCPUTimeMS() { return CPUTimeMS; }
+
 protected:
+	mutable float CPUTimeMS;
+
 	NiagaraEffectRenderer()	
 		: Material(nullptr)
 	{
@@ -87,13 +94,15 @@ private:
 
 
 
+
 /**
 * NiagaraEffectRendererSprites renders an FNiagaraSimulation as sprite particles
 */
 class NiagaraEffectRendererSprites : public NiagaraEffectRenderer
 {
 public:	
-	NiagaraEffectRendererSprites(ERHIFeatureLevel::Type FeatureLevel);
+
+	explicit NiagaraEffectRendererSprites(ERHIFeatureLevel::Type FeatureLevel, UNiagaraEffectRendererProperties *Props);
 	~NiagaraEffectRendererSprites()
 	{
 		ReleaseRenderThreadResources();
@@ -107,20 +116,18 @@ public:
 
 	virtual void GetDynamicMeshElements(const TArray<const FSceneView*>& Views, const FSceneViewFamily& ViewFamily, uint32 VisibilityMap, FMeshElementCollector& Collector, const FNiagaraSceneProxy *SceneProxy) const override;
 
-
-
 	/** Update render data buffer from attributes */
 	FNiagaraDynamicDataBase *GenerateVertexData(const FNiagaraEmitterParticleData &Data) override;
 
-
-
-
 	virtual void SetDynamicData_RenderThread(FNiagaraDynamicDataBase* NewDynamicData) override;
 	int GetDynamicDataSize();
 	bool HasDynamicData();
 
+	UClass *GetPropertiesClass() override { return UNiagaraSpriteRendererProperties::StaticClass(); }
+	
 private:
-	struct FNiagaraDynamicDataSprites* DynamicDataRender;
+	UNiagaraSpriteRendererProperties *Properties;
+	struct FNiagaraDynamicDataSprites *DynamicDataRender;
 	mutable TUniformBuffer<FPrimitiveUniformShaderParameters> WorldSpacePrimitiveUniformBuffer;
 	class FParticleSpriteVertexFactory* VertexFactory;
 };
@@ -129,11 +136,6 @@ private:
 
 
 
-
-
-
-
-
 /**
 * NiagaraEffectRendererRibbon renders an FNiagaraSimulation as a ribbon connecting all particles
 * in order by particle age.
@@ -141,7 +143,7 @@ private:
 class NiagaraEffectRendererRibbon : public NiagaraEffectRenderer
 {
 public:
-	NiagaraEffectRendererRibbon(ERHIFeatureLevel::Type FeatureLevel);
+	NiagaraEffectRendererRibbon(ERHIFeatureLevel::Type FeatureLevel, UNiagaraEffectRendererProperties *Props);
 	~NiagaraEffectRendererRibbon()
 	{
 		ReleaseRenderThreadResources();
@@ -165,7 +167,7 @@ public:
 		NewVertex.Color = FLinearColor(Color);
 		NewVertex.ParticleId = 0;
 		NewVertex.RelativeTime = Age.X;
-		NewVertex.Size = FVector2D(1.0f, 1.0f);
+		NewVertex.Size = FVector2D(Rotation.Y, Rotation.Z);
 		NewVertex.Rotation = Rotation.X;
 		NewVertex.SubImageIndex = 0.f;
 		NewVertex.Tex_U = UV1.X;
@@ -180,11 +182,13 @@ public:
 	int GetDynamicDataSize();
 	bool HasDynamicData();
 
+	UClass *GetPropertiesClass() override { return UNiagaraRibbonRendererProperties::StaticClass(); }
 
 private:
-	struct FNiagaraDynamicDataRibbon* DynamicDataRender;
+	class FParticleBeamTrailVertexFactory *VertexFactory;
+	UNiagaraRibbonRendererProperties *Properties;
+	struct FNiagaraDynamicDataRibbon *DynamicDataRender;
 	mutable TUniformBuffer<FPrimitiveUniformShaderParameters> WorldSpacePrimitiveUniformBuffer;
-	class FParticleBeamTrailVertexFactory* VertexFactory;
 };
 
 
diff --git a/Engine/Source/Runtime/Engine/Classes/Engine/NiagaraEffectRendererProperties.h b/Engine/Source/Runtime/Engine/Classes/Engine/NiagaraEffectRendererProperties.h
new file mode 100644
--- /dev/null
+++ b/Engine/Source/Runtime/Engine/Classes/Engine/NiagaraEffectRendererProperties.h
@@ -0,0 +1,23 @@
+// Copyright 1998-2015 Epic Games, Inc. All Rights Reserved.
+
+#include "NiagaraEffectRendererProperties.generated.h"
+
+#pragma once
+/**
+* Emitter properties base class
+* Each Effectrenderer derives from this with its own class, and returns it in GetProperties; a copy
+* of those specific properties is stored on FNiagaraEmitterProperties (on the effect) for serialization 
+* and handed back to the effect renderer on load.
+*/
+UCLASS()
+class UNiagaraEffectRendererProperties : public UObject
+{
+	GENERATED_UCLASS_BODY()
+	UNiagaraEffectRendererProperties()
+	{}
+
+	UPROPERTY()
+	FName dummy;
+};
+
+
diff --git a/Engine/Source/Runtime/Engine/Classes/Engine/NiagaraRibbonRendererProperties.h b/Engine/Source/Runtime/Engine/Classes/Engine/NiagaraRibbonRendererProperties.h
new file mode 100644
--- /dev/null
+++ b/Engine/Source/Runtime/Engine/Classes/Engine/NiagaraRibbonRendererProperties.h
@@ -0,0 +1,19 @@
+// Copyright 1998-2015 Epic Games, Inc. All Rights Reserved.
+
+#pragma once
+
+#include "NiagaraEffectRendererProperties.h"
+#include "NiagaraRibbonRendererProperties.generated.h"
+
+UCLASS()
+class UNiagaraRibbonRendererProperties : public UNiagaraEffectRendererProperties
+{
+public:
+	GENERATED_UCLASS_BODY()
+
+	UNiagaraRibbonRendererProperties()
+	{
+	}
+
+//	UPROPERTY(EditAnywhere, Category = "Ribbon Rendering")
+};
diff --git a/Engine/Source/Runtime/Engine/Classes/Engine/NiagaraSimulation.h b/Engine/Source/Runtime/Engine/Classes/Engine/NiagaraSimulation.h
--- a/Engine/Source/Runtime/Engine/Classes/Engine/NiagaraSimulation.h
+++ b/Engine/Source/Runtime/Engine/Classes/Engine/NiagaraSimulation.h
@@ -6,6 +6,7 @@ NiagaraSimulation.h: Niagara emitter simulation class
 #pragma once
 
 #include "NiagaraScript.h"
+#include "NiagaraSpriteRendererProperties.h"
 #include "Components/NiagaraComponent.h"
 #include "Runtime/VectorVM/Public/VectorVM.h"
 #include "NiagaraSimulation.generated.h"
@@ -28,6 +29,27 @@ enum EEmitterRenderModuleType
 };
 
 
+class SimpleTimer
+{
+public:
+	SimpleTimer()
+	{
+		StartTime = FPlatformTime::Seconds() * 1000.0;
+	}
+
+	double GetElapsedMilliseconds()
+	{
+		return (FPlatformTime::Seconds()*1000.0) - StartTime;
+	}
+
+	~SimpleTimer()
+	{
+	}
+
+private:
+	double StartTime;
+};
+
 /** 
  *	FniagaraEmitterProperties stores the attributes of an FNiagaraSimulation
  *	that need to be serialized and are used for its initialization 
@@ -37,10 +59,14 @@ struct FNiagaraEmitterProperties
 {
 	GENERATED_USTRUCT_BODY()
 public:
-	FNiagaraEmitterProperties() : bIsEnabled(true), SpawnRate(50), UpdateScript(nullptr), SpawnScript(nullptr), Material(nullptr), RenderModuleType(RMT_Sprites)
+	FNiagaraEmitterProperties() : 
+		bIsEnabled(true), SpawnRate(50), UpdateScript(nullptr), SpawnScript(nullptr), Material(nullptr), 
+		RenderModuleType(RMT_Sprites), RendererProperties(nullptr)
 	{
+		Name = FString(TEXT("New Emitter"));
 	}
-
+	UPROPERTY(EditAnywhere, Category = "Emitter Properties")
+	FString Name;
 	UPROPERTY(EditAnywhere, Category = "Emitter Properties")
 	bool bIsEnabled;
 	UPROPERTY(EditAnywhere, Category = "Emitter Properties")
@@ -54,6 +80,9 @@ public:
 	UPROPERTY(EditAnywhere, Category = "Emitter Properties")
 	TEnumAsByte<EEmitterRenderModuleType> RenderModuleType;
 
+	UPROPERTY()
+	UNiagaraEffectRendererProperties *RendererProperties;
+
 	UPROPERTY(EditAnywhere, Category = "Emitter Properties")
 	FNiagaraConstantMap ExternalConstants;		// these are the update script constants from the effect editor; will be added to the emitter's constant map
 
@@ -96,12 +125,20 @@ public:
 
 	FNiagaraEmitterProperties *GetProperties()	{ return Props; }
 	void SetProperties(FNiagaraEmitterProperties *InProps)	{ Props = InProps; }
+
+	ENGINE_API float GetTotalCPUTime();
+	ENGINE_API int	GetTotalBytesUsed();
+
 private:
 	FNiagaraEmitterProperties *Props;		// points to an entry in the array of FNiagaraProperties stored in the EffectInstance (itself pointing to the effect's properties)
 
+	/* The age of the emitter*/
 	float Age;
+	/* If false, don't tick or render*/
 	bool bIsEnabled;
-
+	/* Seconds taken to process everything (including rendering) */
+	float CPUTimeMS;
+	
 	/** Local constant set. */
 	FNiagaraConstantMap Constants;
 	/** particle simulation data */
diff --git a/Engine/Source/Runtime/Engine/Classes/Engine/NiagaraSpriteRendererProperties.h b/Engine/Source/Runtime/Engine/Classes/Engine/NiagaraSpriteRendererProperties.h
new file mode 100644
--- /dev/null
+++ b/Engine/Source/Runtime/Engine/Classes/Engine/NiagaraSpriteRendererProperties.h
@@ -0,0 +1,21 @@
+// Copyright 1998-2015 Epic Games, Inc. All Rights Reserved.
+
+#pragma once
+
+#include "NiagaraEffectRendererProperties.h"
+#include "NiagaraSpriteRendererProperties.generated.h"
+
+UCLASS()
+class UNiagaraSpriteRendererProperties : public UNiagaraEffectRendererProperties
+{
+public:
+	GENERATED_UCLASS_BODY()
+
+	UNiagaraSpriteRendererProperties()
+	{
+		SubImageInfo = FVector2D(1.0f, 1.0f);
+	}
+
+	UPROPERTY(EditAnywhere, Category = "Sprite Rendering")
+	FVector2D SubImageInfo;
+};
diff --git a/Engine/Source/Runtime/Engine/Private/NiagaraComponent.cpp b/Engine/Source/Runtime/Engine/Private/NiagaraComponent.cpp
--- a/Engine/Source/Runtime/Engine/Private/NiagaraComponent.cpp
+++ b/Engine/Source/Runtime/Engine/Private/NiagaraComponent.cpp
@@ -129,7 +129,10 @@ void FNiagaraSceneProxy::GetDynamicMeshElements(const TArray<const FSceneView*>&
 {
 	for (NiagaraEffectRenderer *Renderer : EffectRenderers)
 	{
-		Renderer->GetDynamicMeshElements(Views, ViewFamily, VisibilityMap, Collector, this);
+		if (Renderer)
+		{
+			Renderer->GetDynamicMeshElements(Views, ViewFamily, VisibilityMap, Collector, this);
+		}
 	}
 }
 
diff --git a/Engine/Source/Runtime/Engine/Private/NiagaraEffect.cpp b/Engine/Source/Runtime/Engine/Private/NiagaraEffect.cpp
--- a/Engine/Source/Runtime/Engine/Private/NiagaraEffect.cpp
+++ b/Engine/Source/Runtime/Engine/Private/NiagaraEffect.cpp
@@ -8,37 +8,66 @@
 UNiagaraEffect::UNiagaraEffect(const FObjectInitializer& ObjectInitializer)
 : Super(ObjectInitializer)
 {
-
 }
 
 
 
-FNiagaraEmitterProperties &UNiagaraEffect::AddEmitterProperties()
+FNiagaraEmitterProperties *UNiagaraEffect::AddEmitterProperties()
 {
-	FNiagaraEmitterProperties Props;
+	FNiagaraEmitterProperties *Props = new FNiagaraEmitterProperties();
 	EmitterProps.Add(Props);
-	return EmitterProps.Top();
+	return Props;
 }
 
 
 
+void UNiagaraEffect::CreateEffectRendererProps(TSharedPtr<FNiagaraSimulation> Sim)
+{
+	UClass *RendererProps = Sim->GetEffectRenderer()->GetPropertiesClass();
+	if (RendererProps)
+	{
+		Sim->GetProperties()->RendererProperties = ConstructObject<UNiagaraEffectRendererProperties>(RendererProps, this);
+	}
+	else
+	{
+		Sim->GetProperties()->RendererProperties = nullptr;
+	}
+}
+
 
 void UNiagaraEffect::PostLoad()
 {
 	Super::PostLoad();
+	EmitterProps.Empty();
+	for (FNiagaraEmitterProperties &Props : EmitterPropsSerialized)
+	{
+		FNiagaraEmitterProperties *NewProps = new FNiagaraEmitterProperties(Props);
+		EmitterProps.Add(NewProps);
+	}
 }
 
 
+void UNiagaraEffect::PreSave()
+{
+	EmitterPropsSerialized.Empty();
+	for (FNiagaraEmitterProperties *Props : EmitterProps)
+	{
+		EmitterPropsSerialized.Add(*Props);
+	}
+}
+
 
 
 
-void FNiagaraEffectInstance::AddEmitter(FNiagaraEmitterProperties *Properties)
+TSharedPtr<FNiagaraSimulation> FNiagaraEffectInstance::AddEmitter(FNiagaraEmitterProperties *Properties)
 {
-	TSharedPtr<FNiagaraSimulation> Sim = MakeShareable(new FNiagaraSimulation(Properties));
+	FNiagaraSimulation *SimPtr = new FNiagaraSimulation(Properties);
+	TSharedPtr<FNiagaraSimulation> Sim = MakeShareable(SimPtr);
 
 	Sim->SetRenderModuleType(RMT_Sprites, Component->GetWorld()->FeatureLevel);
 	Emitters.Add(Sim);
 	FNiagaraSceneProxy *SceneProxy = static_cast<FNiagaraSceneProxy*>(Component->SceneProxy);
 	SceneProxy->UpdateEffectRenderers(this);
 
+	return Sim;
 }
\ No newline at end of file
diff --git a/Engine/Source/Runtime/Engine/Private/NiagaraEffectRenderer.cpp b/Engine/Source/Runtime/Engine/Private/NiagaraEffectRenderer.cpp
--- a/Engine/Source/Runtime/Engine/Private/NiagaraEffectRenderer.cpp
+++ b/Engine/Source/Runtime/Engine/Private/NiagaraEffectRenderer.cpp
@@ -2,14 +2,21 @@
 
 
 #include "EnginePrivate.h"
-#include "Engine/NiagaraEffectRenderer.h"
 #include "Particles/ParticleResources.h"
-
+#include "Engine/NiagaraEffectRenderer.h"
 
 DECLARE_CYCLE_STAT(TEXT("Render Total"), STAT_NiagaraRender, STATGROUP_Niagara);
 DECLARE_CYCLE_STAT(TEXT("Render Sprites"), STAT_NiagaraRenderSprites, STATGROUP_Niagara);
 DECLARE_CYCLE_STAT(TEXT("Render Ribbons"), STAT_NiagaraRenderRibbons, STATGROUP_Niagara);
 
+UNiagaraEffectRendererProperties::UNiagaraEffectRendererProperties(FObjectInitializer const &Initializer)
+{}
+
+UNiagaraSpriteRendererProperties::UNiagaraSpriteRendererProperties(FObjectInitializer const &Initializer)
+{}
+
+UNiagaraRibbonRendererProperties::UNiagaraRibbonRendererProperties(FObjectInitializer const &Initializer)
+{}
 
 
 
@@ -55,11 +62,13 @@ public:
 
 
 
-NiagaraEffectRendererSprites::NiagaraEffectRendererSprites(ERHIFeatureLevel::Type FeatureLevel) :
+NiagaraEffectRendererSprites::NiagaraEffectRendererSprites(ERHIFeatureLevel::Type FeatureLevel, UNiagaraEffectRendererProperties *InProps) :
 	NiagaraEffectRenderer(),
 	DynamicDataRender(NULL)
 {
+	//check(InProps);
 	VertexFactory = new FParticleSpriteVertexFactory(PVFT_Sprite, FeatureLevel);
+	Properties = Cast<UNiagaraSpriteRendererProperties>(InProps);
 }
 
 
@@ -79,6 +88,8 @@ void NiagaraEffectRendererSprites::GetDynamicMeshElements(const TArray<const FSc
 	SCOPE_CYCLE_COUNTER(STAT_NiagaraRender);
 	SCOPE_CYCLE_COUNTER(STAT_NiagaraRenderSprites);
 
+	SimpleTimer MeshElementsTimer;
+
 	check(DynamicDataRender)
 		
 	if (DynamicDataRender->VertexData.Num() == 0)
@@ -128,11 +139,14 @@ void NiagaraEffectRendererSprites::GetDynamicMeshElements(const TArray<const FSc
 				PerViewUniformParameters.RotationBias = 0.0f;
 				PerViewUniformParameters.TangentSelector = FVector4(0.0f, 0.0f, 0.0f, 1.0f);
 				PerViewUniformParameters.InvDeltaSeconds = 0.0f;
-				PerViewUniformParameters.SubImageSize = FVector4(1.0f, 1.0f, 1.0f, 1.0f);
+				if (Properties)
+				{
+					PerViewUniformParameters.SubImageSize = FVector4(Properties->SubImageInfo.X, Properties->SubImageInfo.Y, 1.0f / Properties->SubImageInfo.X, 1.0f / Properties->SubImageInfo.Y);
+				}
 				PerViewUniformParameters.NormalsType = 0;
 				PerViewUniformParameters.NormalsSphereCenter = FVector4(0.0f, 0.0f, 0.0f, 1.0f);
 				PerViewUniformParameters.NormalsCylinderUnitDirection = FVector4(0.0f, 0.0f, 1.0f, 0.0f);
-				PerViewUniformParameters.PivotOffset = FVector2D(0.0f, 0.0f);
+				PerViewUniformParameters.PivotOffset = FVector2D(-0.5f, -0.5f);
 				PerViewUniformParameters.MacroUVParameters = FVector4(0.0f, 0.0f, 1.0f, 1.0f);
 
 				// Collector.AllocateOneFrameResource uses default ctor, initialize the vertex factory
@@ -184,6 +198,8 @@ void NiagaraEffectRendererSprites::GetDynamicMeshElements(const TArray<const FSc
 			}
 		}
 	}
+
+	CPUTimeMS += MeshElementsTimer.GetElapsedMilliseconds();
 }
 
 
@@ -191,6 +207,9 @@ void NiagaraEffectRendererSprites::GetDynamicMeshElements(const TArray<const FSc
 FNiagaraDynamicDataBase *NiagaraEffectRendererSprites::GenerateVertexData(const FNiagaraEmitterParticleData &Data)
 {
 	SCOPE_CYCLE_COUNTER(STAT_NiagaraGenSpriteVertexData);
+
+	SimpleTimer VertexDataTimer;
+
 	FNiagaraDynamicDataSprites *DynamicData = new FNiagaraDynamicDataSprites;
 	TArray<FParticleSpriteVertex>& RenderData = DynamicData->VertexData;
 
@@ -201,7 +220,11 @@ FNiagaraDynamicDataBase *NiagaraEffectRendererSprites::GenerateVertexData(const
 	const FVector4 *ColPtr = Data.GetAttributeData("Color");
 	const FVector4 *AgePtr = Data.GetAttributeData("Age");
 	const FVector4 *RotPtr = Data.GetAttributeData("Rotation");
-
+	uint32 NumSubImages = 1;
+	if (Properties)
+	{
+		NumSubImages = Properties->SubImageInfo.X*Properties->SubImageInfo.Y;
+	}
 	float ParticleId = 0.0f, IdInc = 1.0f / Data.GetNumParticles();
 	RenderData.AddUninitialized(Data.GetNumParticles());
 	for (uint32 ParticleIndex = 0; ParticleIndex < Data.GetNumParticles(); ParticleIndex++)
@@ -213,9 +236,9 @@ FNiagaraDynamicDataBase *NiagaraEffectRendererSprites::GenerateVertexData(const
 		NewVertex.ParticleId = ParticleId;
 		ParticleId += IdInc;
 		NewVertex.RelativeTime = AgePtr[ParticleIndex].X;
-		NewVertex.Size = FVector2D(2.0f, 2.0f);
+		NewVertex.Size = FVector2D(RotPtr[ParticleIndex].Y, RotPtr[ParticleIndex].Z);
 		NewVertex.Rotation = RotPtr[ParticleIndex].X;
-		NewVertex.SubImageIndex = 0.f;
+		NewVertex.SubImageIndex = RotPtr[ParticleIndex].W * NumSubImages;
 
 		FPlatformMisc::Prefetch(PosPtr + ParticleIndex+1);
 		FPlatformMisc::Prefetch(RotPtr + ParticleIndex + 1);
@@ -226,6 +249,9 @@ FNiagaraDynamicDataBase *NiagaraEffectRendererSprites::GenerateVertexData(const
 	}
 	//CachedBounds.ExpandBy(MaxSize);
 
+	CPUTimeMS = VertexDataTimer.GetElapsedMilliseconds();
+
+
 	return DynamicData;
 }
 
@@ -263,11 +289,12 @@ bool NiagaraEffectRendererSprites::HasDynamicData()
 
 
 
-NiagaraEffectRendererRibbon::NiagaraEffectRendererRibbon(ERHIFeatureLevel::Type FeatureLevel) :
+NiagaraEffectRendererRibbon::NiagaraEffectRendererRibbon(ERHIFeatureLevel::Type FeatureLevel, UNiagaraEffectRendererProperties *InProps) :
 NiagaraEffectRenderer(),
 DynamicDataRender(NULL)
 {
 	VertexFactory = new FParticleBeamTrailVertexFactory(PVFT_BeamTrail, FeatureLevel);
+	Properties = Cast<UNiagaraRibbonRendererProperties>(InProps);
 }
 
 
@@ -291,6 +318,8 @@ void NiagaraEffectRendererRibbon::GetDynamicMeshElements(const TArray<const FSce
 	SCOPE_CYCLE_COUNTER(STAT_NiagaraRender);
 	SCOPE_CYCLE_COUNTER(STAT_NiagaraRenderRibbons);
 
+	SimpleTimer MeshElementsTimer;
+
 	check(DynamicDataRender)
 	if (DynamicDataRender->VertexData.Num() == 0)
 	{
@@ -381,6 +410,8 @@ void NiagaraEffectRendererRibbon::GetDynamicMeshElements(const TArray<const FSce
 			}
 		}
 	}
+
+	CPUTimeMS += MeshElementsTimer.GetElapsedMilliseconds();
 }
 
 void NiagaraEffectRendererRibbon::SetDynamicData_RenderThread(FNiagaraDynamicDataBase* NewDynamicData)
@@ -416,6 +447,9 @@ bool NiagaraEffectRendererRibbon::HasDynamicData()
 FNiagaraDynamicDataBase *NiagaraEffectRendererRibbon::GenerateVertexData(const FNiagaraEmitterParticleData &Data)
 {
 	SCOPE_CYCLE_COUNTER(STAT_NiagaraGenRibbonVertexData);
+
+	SimpleTimer VertexDataTimer;
+
 	FNiagaraDynamicDataRibbon *DynamicData = new FNiagaraDynamicDataRibbon;
 	TArray<FParticleBeamTrailVertex>& RenderData = DynamicData->VertexData;
 
@@ -457,15 +491,16 @@ FNiagaraDynamicDataBase *NiagaraEffectRendererRibbon::GenerateVertexData(const F
 		{
 			ParticleDir = PrevDir*0.1f;
 		}
+		FVector NormDir = ParticleDir.GetSafeNormal();
 
-		FVector ParticleRight = FVector::CrossProduct(ParticleDir, FVector(0.0f, 0.0f, 1.0f));
-		ParticleRight.Normalize();
-		ParticleRight *= 3.0f;
+		FVector ParticleRight = FVector::CrossProduct(NormDir, FVector(0.0f, 0.0f, 1.0f));
+		ParticleRight *= RotPtr[Index1].Y;
+		FVector ParticleRightRot = ParticleRight.RotateAngleAxis(RotPtr[Index1].X, NormDir);
 
 		if (i == 0)
 		{
-			AddRibbonVert(RenderData, ParticlePos + ParticleRight, Data, UVs[0], ColorPtr[Index1], AgePtr[Index1], RotPtr[i]);
-			AddRibbonVert(RenderData, ParticlePos - ParticleRight, Data, UVs[1], ColorPtr[Index1], AgePtr[Index1], RotPtr[i]);
+			AddRibbonVert(RenderData, ParticlePos + ParticleRightRot, Data, UVs[0], ColorPtr[Index1], AgePtr[Index1], RotPtr[i]);
+			AddRibbonVert(RenderData, ParticlePos - ParticleRightRot, Data, UVs[1], ColorPtr[Index1], AgePtr[Index1], RotPtr[i]);
 		}
 		else
 		{
@@ -473,12 +508,15 @@ FNiagaraDynamicDataBase *NiagaraEffectRendererRibbon::GenerateVertexData(const F
 			AddRibbonVert(RenderData, PrevPos, Data, UVs[1], ColorPtr[Index1], AgePtr[Index1], RotPtr[i]);
 		}
 
-		AddRibbonVert(RenderData, ParticlePos - ParticleRight + ParticleDir, Data, UVs[2], ColorPtr[Index2], AgePtr[Index2], RotPtr[i]);
-		AddRibbonVert(RenderData, ParticlePos + ParticleRight + ParticleDir, Data, UVs[3], ColorPtr[Index2], AgePtr[Index2], RotPtr[i]);
-		PrevPos = ParticlePos - ParticleRight + ParticleDir;
-		PrevPos2 = ParticlePos + ParticleRight + ParticleDir;
+		ParticleRightRot = ParticleRight.RotateAngleAxis(RotPtr[Index2].X, NormDir);
+		AddRibbonVert(RenderData, ParticlePos - ParticleRightRot + ParticleDir, Data, UVs[2], ColorPtr[Index2], AgePtr[Index2], RotPtr[i]);
+		AddRibbonVert(RenderData, ParticlePos + ParticleRightRot + ParticleDir, Data, UVs[3], ColorPtr[Index2], AgePtr[Index2], RotPtr[i]);
+		PrevPos = ParticlePos - ParticleRightRot + ParticleDir;
+		PrevPos2 = ParticlePos + ParticleRightRot + ParticleDir;
 		PrevDir = ParticleDir;
 	}
 
+	CPUTimeMS = VertexDataTimer.GetElapsedMilliseconds();
+
 	return DynamicData;
-}
\ No newline at end of file
+}
diff --git a/Engine/Source/Runtime/Engine/Private/NiagaraSimulation.cpp b/Engine/Source/Runtime/Engine/Private/NiagaraSimulation.cpp
--- a/Engine/Source/Runtime/Engine/Private/NiagaraSimulation.cpp
+++ b/Engine/Source/Runtime/Engine/Private/NiagaraSimulation.cpp
@@ -29,10 +29,29 @@ FNiagaraSimulation::FNiagaraSimulation(FNiagaraEmitterProperties *InProps, ERHIF
 	SetRenderModuleType(InProps->RenderModuleType, InFeatureLevel);
 }
 
+
+float FNiagaraSimulation::GetTotalCPUTime()
+{
+	float Total = CPUTimeMS;
+	if (EffectRenderer)
+	{
+		Total += EffectRenderer->GetCPUTimeMS();
+	}
+
+	return Total;
+}
+
+int FNiagaraSimulation::GetTotalBytesUsed()
+{
+	return Data.GetBytesUsed();
+}
+
+
 void FNiagaraSimulation::Tick(float DeltaSeconds)
 {
 	SCOPE_CYCLE_COUNTER(STAT_NiagaraTick);
 
+	SimpleTimer TickTime;
 
 	// Cache the ComponentToWorld transform.
 //	CachedComponentToWorld = Component.GetComponentToWorld();
@@ -69,6 +88,8 @@ void FNiagaraSimulation::Tick(float DeltaSeconds)
 		SpawnAndKillParticles(NumToSpawn);
 	}
 
+	CPUTimeMS = TickTime.GetElapsedMilliseconds();
+
 	DECLARE_DWORD_COUNTER_STAT(TEXT("NumParticles"), STAT_NiagaraNumParticles, STATGROUP_Niagara);
 	INC_DWORD_STAT_BY(STAT_NiagaraNumParticles, Data.GetNumParticles());
 
@@ -206,11 +227,11 @@ void FNiagaraSimulation::SetRenderModuleType(EEmitterRenderModuleType Type, ERHI
 		Props->RenderModuleType = Type;
 		switch (Type)
 		{
-		case RMT_Sprites: EffectRenderer = new NiagaraEffectRendererSprites(FeatureLevel);
+		case RMT_Sprites: EffectRenderer = new NiagaraEffectRendererSprites(FeatureLevel, Props->RendererProperties);
 			break;
-		case RMT_Ribbon: EffectRenderer = new NiagaraEffectRendererRibbon(FeatureLevel);
+		case RMT_Ribbon: EffectRenderer = new NiagaraEffectRendererRibbon(FeatureLevel, Props->RendererProperties);
 			break;
-		default:	EffectRenderer = new NiagaraEffectRendererSprites(FeatureLevel);
+		default:	EffectRenderer = new NiagaraEffectRendererSprites(FeatureLevel, Props->RendererProperties);
 					Props->RenderModuleType = RMT_Sprites;
 					break;
 		}
diff --git a/Engine/Source/Runtime/VectorVM/Private/VectorVM.cpp b/Engine/Source/Runtime/VectorVM/Private/VectorVM.cpp
--- a/Engine/Source/Runtime/VectorVM/Private/VectorVM.cpp
+++ b/Engine/Source/Runtime/VectorVM/Private/VectorVM.cpp
@@ -385,7 +385,7 @@ struct FVectorKernelSin : public TUnaryVectorKernel<FVectorKernelSin>
 {
 	static void VM_FORCEINLINE DoKernel(VectorRegister* RESTRICT Dst, VectorRegister Src0)
 	{
-		*Dst = VectorSin(Src0);
+		*Dst = VectorSin( VectorMultiply(Src0, GlobalVectorConstants::Pi));
 	}
 };
 
@@ -393,7 +393,7 @@ struct FVectorKernelCos : public TUnaryVectorKernel<FVectorKernelCos>
 {
 	static void VM_FORCEINLINE DoKernel(VectorRegister* RESTRICT Dst, VectorRegister Src0)
  	{
-		*Dst = VectorCos(Src0);
+		*Dst = VectorCos(VectorMultiply(Src0, GlobalVectorConstants::Pi));
 	}
 };
 
@@ -401,7 +401,7 @@ struct FVectorKernelTan : public TUnaryVectorKernel<FVectorKernelTan>
 {
 	static void VM_FORCEINLINE DoKernel(VectorRegister* RESTRICT Dst, VectorRegister Src0)
 	{
-		*Dst = VectorTan(Src0);
+		*Dst = VectorTan(VectorMultiply(Src0, GlobalVectorConstants::Pi));
 	}
 };
 
@@ -481,23 +481,16 @@ struct FVectorKernelLessThan : public TBinaryVectorKernel<FVectorKernelLessThan>
 {
 	static void FORCEINLINE DoKernel(VectorRegister* RESTRICT Dst, VectorRegister Src0, VectorRegister Src1)
 	{
-		const VectorRegister large = MakeVectorRegister(BIG_NUMBER, BIG_NUMBER, BIG_NUMBER, BIG_NUMBER);
-		const VectorRegister one = MakeVectorRegister(1.0f, 1.0f, 1.0f, 1.0f);
-		const VectorRegister zero = MakeVectorRegister(0.0f, 0.0f, 0.0f, 0.0f);
+		const VectorRegister Large = MakeVectorRegister(BIG_NUMBER, BIG_NUMBER, BIG_NUMBER, BIG_NUMBER);
+		const VectorRegister One = MakeVectorRegister(1.0f, 1.0f, 1.0f, 1.0f);
+		const VectorRegister Zero = MakeVectorRegister(0.0f, 0.0f, 0.0f, 0.0f);
 
 		float const* FloatSrc0 = reinterpret_cast<float const*>(&Src0);
 		float const* FloatSrc1 = reinterpret_cast<float const*>(&Src1);
-		VectorRegister tmp = VectorSubtract(Src1, Src0);
-		tmp = VectorMultiply(tmp, large);
-		tmp = VectorMin(tmp, one);
-		*Dst = VectorMax(tmp, zero);
-		/*
-		float f1 = FloatSrc0[0] < FloatSrc1[0] ? 1.0f : 0.0f;
-		float f2 = FloatSrc0[1] < FloatSrc1[1] ? 1.0f : 0.0f;
-		float f3 = FloatSrc0[2] < FloatSrc1[2] ? 1.0f : 0.0f;
-		float f4 = FloatSrc0[3] < FloatSrc1[3] ? 1.0f : 0.0f;
-		*Dst = MakeVectorRegister(f1, f2, f3, f4);
-		*/
+		VectorRegister Tmp = VectorSubtract(Src1, Src0);
+		Tmp = VectorMultiply(Tmp, Large);
+		Tmp = VectorMin(Tmp, One);
+		*Dst = VectorMax(Tmp, Zero);
 	}
 };
 
@@ -546,7 +539,6 @@ struct FVectorKernelRandom : public TUnaryVectorKernel<FVectorKernelRandom>
 			static_cast<float>(FMath::Rand()) / rm,
 			static_cast<float>(FMath::Rand()) / rm);
 		*Dst = VectorMultiply(Result, Src0);
-		
 	}
 };
 
@@ -801,6 +793,7 @@ void VectorVM::Exec(
 			case EOp::composey: FVectorKernelCompose<1, 1, 1, 1>::Exec(Context); break;
 			case EOp::composez: FVectorKernelCompose<2, 2, 2, 2>::Exec(Context); break;
 			case EOp::composew: FVectorKernelCompose<3, 3, 3, 3>::Exec(Context); break;
+			case EOp::lessthan: FVectorKernelLessThan::Exec(Context); break;
 			case EOp::output: FVectorKernelOutput::Exec(Context); break;
 
 			// Execution always terminates with a "done" opcode.
diff --git a/Engine/Source/Runtime/VectorVM/Public/VectorVM.h b/Engine/Source/Runtime/VectorVM/Public/VectorVM.h
--- a/Engine/Source/Runtime/VectorVM/Public/VectorVM.h
+++ b/Engine/Source/Runtime/VectorVM/Public/VectorVM.h
@@ -72,6 +72,7 @@ namespace VectorVM
 			composez,
 			composew,
 			output,
+			lessthan,
 			NumOpcodes
 	};
 
