build fix #2

diff --git a/Engine/Source/Editor/NiagaraEditor/Private/SNiagaraEffectEditorWidget.cpp b/Engine/Source/Editor/NiagaraEditor/Private/SNiagaraEffectEditorWidget.cpp
--- a/Engine/Source/Editor/NiagaraEditor/Private/SNiagaraEffectEditorWidget.cpp
+++ b/Engine/Source/Editor/NiagaraEditor/Private/SNiagaraEffectEditorWidget.cpp
@@ -5,7 +5,7 @@
 #include "SNotificationList.h"
 #include "SExpandableArea.h"
 #include "Editor/PropertyEditor/Public/PropertyEditorModule.h"
-#include "Editor/PropertyEditor/Public/IStructureDetailsView.h"
+#include "Editor/PropertyEditor/Public/IDetailsView.h"
 
 
 #define LOCTEXT_NAMESPACE "NiagaraEffectEditor"
