Fixed code relying on SLATE_TEXT_ATTRIBUTE

UETOOL-213 - Minimize Slate FString -> FText conversion (remove SLATE_TEXT_ATTRIBUTE)

New code found post-update.

diff --git a/Engine/Source/Editor/NiagaraEditor/Private/SNiagaraEffectEditorWidget.cpp b/Engine/Source/Editor/NiagaraEditor/Private/SNiagaraEffectEditorWidget.cpp
--- a/Engine/Source/Editor/NiagaraEditor/Private/SNiagaraEffectEditorWidget.cpp
+++ b/Engine/Source/Editor/NiagaraEditor/Private/SNiagaraEffectEditorWidget.cpp
@@ -314,7 +314,7 @@ void SEmitterWidget::Construct(const FArguments& InArgs)
 										.AutoHeight()
 										.Expose(UpdateScriptCurveConstantListSlot)
 										[
-											SNew(SExpandableArea).InitiallyCollapsed(true).AreaTitle(FString("Curves"))
+											SNew(SExpandableArea).InitiallyCollapsed(true).AreaTitle(LOCTEXT("CurvesLabel", "Curves"))
 											.BodyContent()
 											[
 												SNew(SListView<TSharedPtr<EditorExposedVectorCurveConstant>>)
