New niagara graph nodes.

Input - General input node that replaces the GetAttribute and GetConstant node. The input node can do both depending on the scripts usage. The input will also be used to connect up scripts being used as functions.
Output - Replaces NiagaraNodeOutputUpdate. Allows any number of outputs with user set names and types. Also used to connect up outputs of scripts used as functions.

Moved all constants and attributes to use a Type,Name pair as an ID rather than jsut an FName. This is required for full support of attributes and constant of any type. Still some work to be done for full support of this.

A few minor unrelated fixes.

diff --git a/Engine/Source/Editor/NiagaraEditor/Private/NiagaraCompiler.cpp b/Engine/Source/Editor/NiagaraEditor/Private/NiagaraCompiler.cpp
--- a/Engine/Source/Editor/NiagaraEditor/Private/NiagaraCompiler.cpp
+++ b/Engine/Source/Editor/NiagaraEditor/Private/NiagaraCompiler.cpp
@@ -38,9 +38,8 @@ class FNiagaraExpression_GetAttribute : public FNiagaraExpression
 {
 public:
 
-	FNiagaraExpression_GetAttribute(class FNiagaraCompiler* InCompiler, ENiagaraDataType InType, FName InName)
-		: FNiagaraExpression(InCompiler, InType)
-		, AttributeName(InName)
+	FNiagaraExpression_GetAttribute(class FNiagaraCompiler* InCompiler, const FNiagaraVariableInfo& InAttribute)
+		: FNiagaraExpression(InCompiler, InAttribute)
 	{
 		ResultLocation = ENiagaraExpressionResultLocation::InputData;
 	}
@@ -48,10 +47,8 @@ public:
 	virtual void Process()override
 	{
 		check(ResultLocation == ENiagaraExpressionResultLocation::InputData);
-		ResultIndex = Compiler->GetAttributeIndex(AttributeName);
+		ResultIndex = Compiler->GetAttributeIndex(Result);
 	}
-
-	FName AttributeName;
 };
 
 /** Expression that gets a constant. */
@@ -59,34 +56,35 @@ class FNiagaraExpression_GetConstant : public FNiagaraExpression
 {
 public:
 
-	FNiagaraExpression_GetConstant(class FNiagaraCompiler* InCompiler, ENiagaraDataType InType, FName InName, bool bIsInternal)
-		: FNiagaraExpression(InCompiler, InType)
-		, ConstantName(InName)
+	FNiagaraExpression_GetConstant(class FNiagaraCompiler* InCompiler, const FNiagaraVariableInfo& InConstant, bool bIsInternal)
+		: FNiagaraExpression(InCompiler, InConstant)
 		, bInternal(bIsInternal)
 	{
 		ResultLocation = ENiagaraExpressionResultLocation::Constants;
 		//For matrix constants we must add 4 input expressions that will be filled in as the constant is processed.
 		//They must at least exist now so that other expressions can reference them.
-		if (ResultType == ENiagaraDataType::Matrix)
-		{
-			SourceExpressions.Add(MakeShareable(new FNiagaraExpression(InCompiler, ENiagaraDataType::Vector)));
-			SourceExpressions.Add(MakeShareable(new FNiagaraExpression(InCompiler, ENiagaraDataType::Vector)));
-			SourceExpressions.Add(MakeShareable(new FNiagaraExpression(InCompiler, ENiagaraDataType::Vector)));
-			SourceExpressions.Add(MakeShareable(new FNiagaraExpression(InCompiler, ENiagaraDataType::Vector)));
-		}
-		else if (ResultType == ENiagaraDataType::Curve)
+		if (Result.Type == ENiagaraDataType::Matrix)
 		{
-			SourceExpressions.Add(MakeShareable(new FNiagaraExpression(InCompiler, ENiagaraDataType::Curve)));
+			static const FName ResultName(TEXT("MatrixComponent"));
+			SourceExpressions.Add(MakeShareable(new FNiagaraExpression(InCompiler, FNiagaraVariableInfo(ResultName, ENiagaraDataType::Vector))));
+			SourceExpressions.Add(MakeShareable(new FNiagaraExpression(InCompiler, FNiagaraVariableInfo(ResultName, ENiagaraDataType::Vector))));
+			SourceExpressions.Add(MakeShareable(new FNiagaraExpression(InCompiler, FNiagaraVariableInfo(ResultName, ENiagaraDataType::Vector))));
+			SourceExpressions.Add(MakeShareable(new FNiagaraExpression(InCompiler, FNiagaraVariableInfo(ResultName, ENiagaraDataType::Vector))));
 		}
+// 		else if (Result.Type == ENiagaraDataType::Curve)
+// 		{
+// 			static const FName ResultName(TEXT("Curve"));
+// 			SourceExpressions.Add(MakeShareable(new FNiagaraExpression(InCompiler, FNiagaraVariableInfo(ResultName, ENiagaraDataType::Curve))));
+// 		}
 	}
 
 	virtual void Process()override
 	{
 		check(ResultLocation == ENiagaraExpressionResultLocation::Constants);
 		//Get the index of the constant. Also gets a component index if this is for a packed scalar constant.
-		Compiler->GetConstantResultIndex(ConstantName, bInternal, ResultIndex, ComponentIndex);
+		Compiler->GetConstantResultIndex(Result, bInternal, ResultIndex, ComponentIndex);
 		
-		if (ResultType == ENiagaraDataType::Matrix)
+		if (Result.Type == ENiagaraDataType::Matrix)
 		{
 			check(SourceExpressions.Num() == 4);
 			TNiagaraExprPtr Src0 = GetSourceExpression(0);
@@ -102,13 +100,12 @@ public:
 			Src3->ResultLocation = ENiagaraExpressionResultLocation::Constants;
 			Src3->ResultIndex = ResultIndex + 3;
 		}
-		else if (ResultType == ENiagaraDataType::Curve)
+		else if (Result.Type == ENiagaraDataType::Curve)
 		{
 			ResultLocation = ENiagaraExpressionResultLocation::BufferConstants;
 		}
 	}
 
-	FName ConstantName;
 	bool bInternal;
 };
 
@@ -117,8 +114,8 @@ class FNiagaraExpression_Collection : public FNiagaraExpression
 {
 public:
 
-	FNiagaraExpression_Collection(class FNiagaraCompiler* InCompiler, ENiagaraDataType Type, TArray<TNiagaraExprPtr>& InSourceExpressions)
-		: FNiagaraExpression(InCompiler, Type)
+	FNiagaraExpression_Collection(class FNiagaraCompiler* InCompiler, const FNiagaraVariableInfo& Result, TArray<TNiagaraExprPtr>& InSourceExpressions)
+		: FNiagaraExpression(InCompiler, Result)
 	{
 		SourceExpressions = InSourceExpressions;
 	}
@@ -135,9 +132,9 @@ void FNiagaraCompiler::CheckInputs(FName OpName, TArray<TNiagaraExprPtr>& Inputs
 	check(OpInfo->Inputs.Num() == NumInputs);
 	for (int32 i = 0; i < NumInputs ; ++i)
 	{
-		if (!(OpInfo->Inputs[i].DataType == Inputs[i]->ResultType || Inputs[i]->ResultType == ENiagaraDataType::Unknown))
+		if (OpInfo->Inputs[i].DataType != Inputs[i]->Result.Type)
 		{
-			UE_LOG(LogNiagaraCompiler, Error, TEXT("Exression %s has incorrect inputs!\nExpected: %d - Actual: %d"), *OpName.ToString(), (int32)OpInfo->Inputs[i].DataType, (int32)((TNiagaraExprPtr)Inputs[i])->ResultType);
+			UE_LOG(LogNiagaraCompiler, Error, TEXT("Exression %s has incorrect inputs!\nExpected: %d - Actual: %d"), *OpName.ToString(), (int32)OpInfo->Inputs[i].DataType, (int32)((TNiagaraExprPtr)Inputs[i])->Result.Type.GetValue());
 		}
 	}
 }
@@ -151,16 +148,16 @@ void FNiagaraCompiler::CheckOutputs(FName OpName, TArray<TNiagaraExprPtr>& Outpu
 	check(OpInfo->Outputs.Num() == NumOutputs);
 	for (int32 i = 0; i < NumOutputs; ++i)
 	{
-		if (!(OpInfo->Outputs[i].DataType == Outputs[i]->ResultType || Outputs[i]->ResultType == ENiagaraDataType::Unknown))
+		if (OpInfo->Outputs[i].DataType != Outputs[i]->Result.Type)
 		{
-			UE_LOG(LogNiagaraCompiler, Error, TEXT("Exression %s has incorrect outputs!\nExpected: %d - Actual: %d"), *OpName.ToString(), (int32)OpInfo->Outputs[i].DataType, (int32)((TNiagaraExprPtr)Outputs[i])->ResultType);
+			UE_LOG(LogNiagaraCompiler, Error, TEXT("Exression %s has incorrect outputs!\nExpected: %d - Actual: %d"), *OpName.ToString(), (int32)OpInfo->Outputs[i].DataType, (int32)((TNiagaraExprPtr)Outputs[i])->Result.Type.GetValue());
 		}
 	}
 }
 
-int32 FNiagaraCompiler::GetAttributeIndex(FName Name)
+int32 FNiagaraCompiler::GetAttributeIndex(const FNiagaraVariableInfo& Attr)const
 {
-	return Attributes.Find(Name);
+	return SourceGraph->GetAttributeIndex(Attr);
 }
 
 TNiagaraExprPtr FNiagaraCompiler::CompileOutputPin(UEdGraphPin* Pin)
@@ -206,36 +203,44 @@ TNiagaraExprPtr FNiagaraCompiler::CompilePin(UEdGraphPin* Pin)
 		{
 			//No connections to this input so add the default as a const expression.
 			const UEdGraphSchema_Niagara* Schema = Cast<const UEdGraphSchema_Niagara>(Pin->GetSchema());
-			ENiagaraDataType DataType = Schema->GetPinDataType(Pin);
+			ENiagaraDataType PinType = Schema->GetPinType(Pin);
 
 			FString ConstString = Pin->GetDefaultAsString();
 			FName ConstName(*ConstString);
-			
-			if (DataType == ENiagaraDataType::Scalar)
+
+			FNiagaraVariableInfo Var(ConstName, PinType);
+
+			switch (PinType)
+			{
+			case ENiagaraDataType::Scalar:
 			{
 				float Default;
 				Schema->GetPinDefaultValue(Pin, Default);
-				ConstantData.SetOrAddInternal(ConstName, Default);
+				ConstantData.SetOrAddInternal(Var, Default);
 			}
-			else if (DataType == ENiagaraDataType::Vector)
+				break;
+			case ENiagaraDataType::Vector:
 			{
 				FVector4 Default;
 				Schema->GetPinDefaultValue(Pin, Default);
-				ConstantData.SetOrAddInternal(ConstName, Default);
-
+				ConstantData.SetOrAddInternal(Var, Default);
 			}
-			else if (DataType == ENiagaraDataType::Matrix)
+				break;
+			case ENiagaraDataType::Matrix:
 			{
 				FMatrix Default;
 				Schema->GetPinDefaultValue(Pin, Default);
-				ConstantData.SetOrAddInternal(ConstName, Default);
+				ConstantData.SetOrAddInternal(Var, Default);
 			}
-			else if (DataType == ENiagaraDataType::Curve)
+				break;
+			case ENiagaraDataType::Curve:
 			{
 				FNiagaraDataObject *Default = nullptr;
-				ConstantData.SetOrAddInternal(ConstName, Default);
+				ConstantData.SetOrAddInternal(Var, Default);
+			}
+				break;
 			}
-			Ret = Expression_GetInternalConstant(ConstName, DataType);
+			Ret = Expression_GetInternalConstant(Var);
 		}
 	}
 	else
@@ -246,63 +251,63 @@ TNiagaraExprPtr FNiagaraCompiler::CompilePin(UEdGraphPin* Pin)
 	return Ret;
 }
 
-void FNiagaraCompiler::GetParticleAttributes(TArray<FName>& OutAttributes)
+void FNiagaraCompiler::GetParticleAttributes(TArray<FNiagaraVariableInfo>& OutAttributes)const
 {
-	check(Source);
-	Source->GetParticleAttributes(OutAttributes);
+	check(SourceGraph);
+	SourceGraph->GetAttributes(OutAttributes);
 }
 
-TNiagaraExprPtr FNiagaraCompiler::Expression_GetAttribute(FName AttributeName)
+TNiagaraExprPtr FNiagaraCompiler::Expression_GetAttribute(const FNiagaraVariableInfo& Attribute)
 {
-	int32 Index = Expressions.Add(MakeShareable(new FNiagaraExpression_GetAttribute(this, ENiagaraDataType::Vector, AttributeName)));
+	int32 Index = Expressions.Add(MakeShareable(new FNiagaraExpression_GetAttribute(this, Attribute)));
 	return Expressions[Index];
 }
 
-TNiagaraExprPtr FNiagaraCompiler::Expression_GetExternalConstant(FName ConstantName, ENiagaraDataType Type)
+TNiagaraExprPtr FNiagaraCompiler::Expression_GetExternalConstant(const FNiagaraVariableInfo& Constant)
 {
-	int32 Index = Expressions.Add(MakeShareable(new FNiagaraExpression_GetConstant(this, Type, ConstantName, false)));
+	int32 Index = Expressions.Add(MakeShareable(new FNiagaraExpression_GetConstant(this, Constant, false)));
 	return Expressions[Index];
 }
 
-TNiagaraExprPtr FNiagaraCompiler::Expression_GetInternalConstant(FName ConstantName, ENiagaraDataType Type)
+TNiagaraExprPtr FNiagaraCompiler::Expression_GetInternalConstant(const FNiagaraVariableInfo& Constant)
 {
-	int32 Index = Expressions.Add(MakeShareable(new FNiagaraExpression_GetConstant(this, Type, ConstantName, true)));
+	int32 Index = Expressions.Add(MakeShareable(new FNiagaraExpression_GetConstant(this, Constant, true)));
 	return Expressions[Index];
 }
 
 TNiagaraExprPtr FNiagaraCompiler::Expression_Collection(TArray<TNiagaraExprPtr>& SourceExpressions)
 {
-	int32 Index = Expressions.Add(MakeShareable(new FNiagaraExpression_Collection(this, ENiagaraDataType::Unknown, SourceExpressions)));
+	//Todo - Collection expressions are just to collect parts of a matrix currently. Its a crappy way to do things that should be replaced.
+	//Definitely don't start using it for collections of other things.
+	int32 Index = Expressions.Add(MakeShareable(new FNiagaraExpression_Collection(this, FNiagaraVariableInfo(NAME_None, ENiagaraDataType::Vector), SourceExpressions)));
 	return Expressions[Index];
 }
 
-TNiagaraExprPtr FNiagaraCompiler::GetAttribute(FName AttributeName)
+TNiagaraExprPtr FNiagaraCompiler::GetAttribute(const FNiagaraVariableInfo& Attribute)
 {
-	return Expression_GetAttribute(AttributeName);
+	return Expression_GetAttribute(Attribute);
 }
 
-TNiagaraExprPtr FNiagaraCompiler::GetExternalConstant(FName ConstantName, float Default)
+TNiagaraExprPtr FNiagaraCompiler::GetExternalConstant(const FNiagaraVariableInfo& Constant, float Default)
 {
-	SetOrAddConstant(false, ConstantName, Default);
-	return Expression_GetExternalConstant(ConstantName, ENiagaraDataType::Scalar);
+	SetOrAddConstant(false, Constant, Default);
+	return Expression_GetExternalConstant(Constant);
 }
-
-TNiagaraExprPtr FNiagaraCompiler::GetExternalConstant(FName ConstantName, FVector4 Default)
+TNiagaraExprPtr FNiagaraCompiler::GetExternalConstant(const FNiagaraVariableInfo& Constant, const FVector4& Default)
 {
-	SetOrAddConstant(false, ConstantName, Default);
-	return Expression_GetExternalConstant(ConstantName, ENiagaraDataType::Vector);
+	SetOrAddConstant(false, Constant, Default);
+	return Expression_GetExternalConstant(Constant);
 }
-
-TNiagaraExprPtr FNiagaraCompiler::GetExternalConstant(FName ConstantName, const FMatrix& Default)
+TNiagaraExprPtr FNiagaraCompiler::GetExternalConstant(const FNiagaraVariableInfo& Constant, const FMatrix& Default)
 {
-	SetOrAddConstant(false, ConstantName, Default);
-	return Expression_GetExternalConstant(ConstantName, ENiagaraDataType::Matrix);
+	SetOrAddConstant(false, Constant, Default);
+	return Expression_GetExternalConstant(Constant);
 }
 
-TNiagaraExprPtr FNiagaraCompiler::GetExternalCurveConstant(FName ConstantName)
+TNiagaraExprPtr FNiagaraCompiler::GetExternalCurveConstant(const FNiagaraVariableInfo& Constant)
 {
-	SetOrAddConstant<FNiagaraDataObject*>(false, ConstantName, nullptr);
-	return Expression_GetExternalConstant(ConstantName, ENiagaraDataType::Curve);
+	SetOrAddConstant<FNiagaraDataObject*>(false, Constant, nullptr);
+	return Expression_GetExternalConstant(Constant);
 }
 
 #undef LOCTEXT_NAMESPACE
\ No newline at end of file
diff --git a/Engine/Source/Editor/NiagaraEditor/Private/NiagaraCompiler.h b/Engine/Source/Editor/NiagaraEditor/Private/NiagaraCompiler.h
--- a/Engine/Source/Editor/NiagaraEditor/Private/NiagaraCompiler.h
+++ b/Engine/Source/Editor/NiagaraEditor/Private/NiagaraCompiler.h
@@ -16,12 +16,12 @@ protected:
 	/** The source of the script we are compiling. */
 	UNiagaraScriptSource* Source;
 
+	/** The source graph of the script we are compiling. */
+	UNiagaraGraph* SourceGraph;
+
 	/** The set of expressions generated from the script source. */
 	TArray<TNiagaraExprPtr> Expressions;
 
-	/** The particle attribute names. */
-	TArray<FName> Attributes;
-	
 	//All internal and external constants used in the graph.
 	FNiagaraScriptConstantData ConstantData;
 
@@ -35,40 +35,41 @@ protected:
 	TNiagaraExprPtr CompileOutputPin(UEdGraphPin* Pin);
 
 	/** Creates an expression retrieving a particle attribute. */
-	TNiagaraExprPtr Expression_GetAttribute(FName AttributeName);
+	TNiagaraExprPtr Expression_GetAttribute(const FNiagaraVariableInfo& Attribute);
 
 	/** Creates an expression retrieving an external constant. */
-	TNiagaraExprPtr Expression_GetExternalConstant(FName ConstantName, ENiagaraDataType Type);
+	TNiagaraExprPtr Expression_GetExternalConstant(const FNiagaraVariableInfo& Constant);
 
 	/** Creates an expression retrieving an internal constant. */
-	TNiagaraExprPtr Expression_GetInternalConstant(FName ConstantName, ENiagaraDataType Type);
+	TNiagaraExprPtr Expression_GetInternalConstant(const FNiagaraVariableInfo& Constant);
 
 	/** Creates an expression collecting some source expressions together for use by future expressions. */
 	TNiagaraExprPtr Expression_Collection(TArray<TNiagaraExprPtr>& SourceExpressions);
 
 	/** Sets the value of a constant. Overwrites the value if it already exists. */
 	template< typename T >
-	void SetOrAddConstant(bool bInternal, FName Name, const T& Default);
+	void SetOrAddConstant(bool bInternal, const FNiagaraVariableInfo& Constant, const T& Default);
 
 public:
 
 	//Begin INiagaraCompiler Interface
 	virtual TNiagaraExprPtr CompilePin(UEdGraphPin* Pin)override;
-	virtual void GetParticleAttributes(TArray<FName>& OutAttributes)override;
-	virtual TNiagaraExprPtr GetAttribute(FName AttributeName)override;
-	virtual TNiagaraExprPtr GetExternalConstant(FName ConstantName, float Default)override;
-	virtual TNiagaraExprPtr GetExternalConstant(FName ConstantName, FVector4 Default)override;
-	virtual TNiagaraExprPtr GetExternalConstant(FName ConstantName, const FMatrix& Default)override;
-	virtual TNiagaraExprPtr GetExternalCurveConstant(FName ConstantName)override;
+
+	virtual TNiagaraExprPtr GetAttribute(const FNiagaraVariableInfo& Attribute)override;
+	virtual TNiagaraExprPtr GetExternalConstant(const FNiagaraVariableInfo& Constant, float Default)override;
+	virtual TNiagaraExprPtr GetExternalConstant(const FNiagaraVariableInfo& Constant, const FVector4& Default)override;
+	virtual TNiagaraExprPtr GetExternalConstant(const FNiagaraVariableInfo& Constant, const FMatrix& Default)override;
+	virtual TNiagaraExprPtr GetExternalCurveConstant(const FNiagaraVariableInfo& Constant)override;
 
 	virtual void CheckInputs(FName OpName, TArray<TNiagaraExprPtr>& Inputs)override;
 	virtual void CheckOutputs(FName OpName, TArray<TNiagaraExprPtr>& Outputs)override;
 	//End INiagaraCompiler Interface
 	
 	/** Gets the index into a constants table of the constant specified by Name and bInternal. */
-	virtual ENiagaraDataType GetConstantResultIndex(FName Name, bool bInternal, int32& OutResultIndex, int32& OutComponentIndex) = 0;
-	/**	Gets the index of a named attribute into either the input or output data. */
-	int32 GetAttributeIndex(FName Name);
+	virtual ENiagaraDataType GetConstantResultIndex(const FNiagaraVariableInfo& Constant, bool bInternal, int32& OutResultIndex, int32& OutComponentIndex) = 0;
+	/**	Gets the index of an attribute. */
+	int32 GetAttributeIndex(const FNiagaraVariableInfo& Attr)const;
+	virtual void GetParticleAttributes(TArray<FNiagaraVariableInfo>& OutAttributes)const;
 
 	/**	Gets the index of a free temporary location. */
 	virtual int32 AquireTemporary() = 0;
@@ -76,11 +77,11 @@ public:
 	virtual void FreeTemporary(int32 TempIndex) = 0;
 };
 
-template< typename T >
-void FNiagaraCompiler::SetOrAddConstant(bool bInternal, FName Name, const T& Default)
+template<typename T>
+void FNiagaraCompiler::SetOrAddConstant(bool bInternal, const FNiagaraVariableInfo& Constant, const T& Default)
 {
 	if (bInternal)
-		ConstantData.SetOrAddInternal(Name, Default);
+		ConstantData.SetOrAddInternal(Constant, Default);
 	else
-		ConstantData.SetOrAddExternal(Name, Default);
+		ConstantData.SetOrAddExternal(Constant, Default);
 }
\ No newline at end of file
diff --git a/Engine/Source/Editor/NiagaraEditor/Private/NiagaraCompiler_VectorVM.cpp b/Engine/Source/Editor/NiagaraEditor/Private/NiagaraCompiler_VectorVM.cpp
--- a/Engine/Source/Editor/NiagaraEditor/Private/NiagaraCompiler_VectorVM.cpp
+++ b/Engine/Source/Editor/NiagaraEditor/Private/NiagaraCompiler_VectorVM.cpp
@@ -19,8 +19,8 @@ class FNiagaraExpression_VMOperation : public FNiagaraExpression
 public:
 	VectorVM::EOp OpCode;
 
-	FNiagaraExpression_VMOperation(FNiagaraCompiler* Compiler, ENiagaraDataType InType, TArray<TNiagaraExprPtr>& InputExpressions, VectorVM::EOp Op)
-		: FNiagaraExpression(Compiler, InType)
+	FNiagaraExpression_VMOperation(FNiagaraCompiler* Compiler, const FNiagaraVariableInfo& InResult, TArray<TNiagaraExprPtr>& InputExpressions, VectorVM::EOp Op)
+		: FNiagaraExpression(Compiler, InResult)
 		, OpCode(Op)
 	{
 		SourceExpressions = InputExpressions;
@@ -92,22 +92,19 @@ class FNiagaraExpression_VMOutput : public FNiagaraExpression
 {
 public:
 
-	FNiagaraExpression_VMOutput(class FNiagaraCompiler* InCompiler, ENiagaraDataType InType, FName InName, TNiagaraExprPtr& InSourceExpression)
-		: FNiagaraExpression(InCompiler, InType)
-		, AttributeName(InName)
+	FNiagaraExpression_VMOutput(class FNiagaraCompiler* InCompiler, const FNiagaraVariableInfo& InAttribute, TNiagaraExprPtr& InSourceExpression)
+		: FNiagaraExpression(InCompiler, InAttribute)
 	{
 		ResultLocation = ENiagaraExpressionResultLocation::OutputData;
 		SourceExpressions.Add(InSourceExpression);
 		check(SourceExpressions.Num() == 1);
 	}
 
-	FName AttributeName;
-
 	virtual void Process()override
 	{
 		FNiagaraCompiler_VectorVM* VMCompiler = (FNiagaraCompiler_VectorVM*)Compiler;
 		check(ResultLocation == ENiagaraExpressionResultLocation::OutputData);
-		ResultIndex = VMCompiler->GetAttributeIndex(AttributeName);
+		ResultIndex = VMCompiler->GetAttributeIndex(Result);
 		
 		FNiagaraExpression* SrcExpr = GetSourceExpression(0).Get();
 		check(SourceExpressions.Num() == 1);
@@ -160,10 +157,10 @@ uint8 FNiagaraCompiler_VectorVM::GetResultVMIndex(FNiagaraExpression* Expression
 	}
 }
 
-ENiagaraDataType FNiagaraCompiler_VectorVM::GetConstantResultIndex(FName Name, bool bInternal, int32& OutResultIndex, int32& OutComponentIndex)
+ENiagaraDataType FNiagaraCompiler_VectorVM::GetConstantResultIndex(const FNiagaraVariableInfo& Constant, bool bInternal, int32& OutResultIndex, int32& OutComponentIndex)
 {
 	ENiagaraDataType Type;
-	ConstantData.GetTableIndex(Name, bInternal, OutResultIndex, OutComponentIndex, Type);
+	ConstantData.GetTableIndex(Constant, bInternal, OutResultIndex, OutComponentIndex, Type);
 	return Type;
 }
 
@@ -186,9 +183,9 @@ void FNiagaraCompiler_VectorVM::WriteCode(uint8 InCode)
 	Script->ByteCode.Add(InCode);
 }
 
-TNiagaraExprPtr FNiagaraCompiler_VectorVM::Output(FName OutputName, TNiagaraExprPtr& SourceExpression)
+TNiagaraExprPtr FNiagaraCompiler_VectorVM::Output(const FNiagaraVariableInfo& Attr, TNiagaraExprPtr& SourceExpression)
 {
-	int32 Index = Expressions.Add(MakeShareable(new FNiagaraExpression_VMOutput(this, ENiagaraDataType::Vector, OutputName, SourceExpression)));
+	int32 Index = Expressions.Add(MakeShareable(new FNiagaraExpression_VMOutput(this, Attr, SourceExpression)));
 	return Expressions[Index];
 }
 
@@ -202,30 +199,13 @@ void FNiagaraCompiler_VectorVM::CompileScript(UNiagaraScript* InScript)
 	Source = CastChecked<UNiagaraScriptSource>(InScript->Source);
 
 	// Clone the source graph so we can modify it as needed; merging in the child graphs
-	UEdGraph* UpdateGraph = FEdGraphUtilities::CloneGraph(Source->UpdateGraph, Source, &MessageLog, true);
-	FEdGraphUtilities::MergeChildrenGraphsIn(UpdateGraph, UpdateGraph, /*bRequireSchemaMatch=*/ true);
+	SourceGraph = CastChecked<UNiagaraGraph>(FEdGraphUtilities::CloneGraph(Source->NodeGraph, Source, &MessageLog));
+	FEdGraphUtilities::MergeChildrenGraphsIn(SourceGraph, SourceGraph, /*bRequireSchemaMatch=*/ true);
 
 	TempRegisters.AddZeroed(VectorVM::NumTempRegisters);
 
 	// Find the output node.
-	UNiagaraNodeOutputUpdate* OutputNode = NULL;
-	{
-		TArray<UNiagaraNodeOutputUpdate*> OutputNodes;
-		Source->UpdateGraph->GetNodesOfClass(OutputNodes);
-		if (OutputNodes.Num() != 1)
-		{
-			UE_LOG(LogNiagaraCompiler_VectorVM, Error, TEXT("Script contains %s output nodes: %s"),
-				OutputNodes.Num() == 0 ? TEXT("no") : TEXT("too many"),
-				*InScript->GetPathName()
-				);
-			return;
-		}
-		OutputNode = OutputNodes[0];
-	}
-	check(OutputNode);
-
-	//Todo - Replace with attributes defined by an Emitter object.
-	Source->GetParticleAttributes(Attributes);
+	UNiagaraNodeOutput* OutputNode = SourceGraph->FindOutputNode();
 
 	TArray<FNiagaraNodeResult> OutputExpressions;
 	OutputNode->Compile(this, OutputExpressions);
@@ -257,13 +237,13 @@ void FNiagaraCompiler_VectorVM::CompileScript(UNiagaraScript* InScript)
 	Script->ByteCode.Add((int8)VectorVM::EOp::done);
 
 	Script->ConstantData = ConstantData;
-	Script->Attributes = Attributes;
-	//To do - scalar and matrix attributes.
+	Script->Attributes = OutputNode->Outputs;
 }
 
 TNiagaraExprPtr FNiagaraCompiler_VectorVM::Expression_VMNative(VectorVM::EOp Op, TArray<TNiagaraExprPtr>& InputExpressions)
 {
-	int32 Index = Expressions.Add(MakeShareable(new FNiagaraExpression_VMOperation(this, ENiagaraDataType::Vector, InputExpressions, Op)));
+	static const FName OpName(TEXT("VMOp"));
+	int32 Index = Expressions.Add(MakeShareable(new FNiagaraExpression_VMOperation(this, FNiagaraVariableInfo(OpName, ENiagaraDataType::Vector), InputExpressions, Op)));
 	return Expressions[Index];
 }
 
@@ -495,7 +475,7 @@ void FNiagaraCompiler_VectorVM::Matrix_Internal(TArray<TNiagaraExprPtr>& InputEx
 {
 	//Just collect the matrix rows into an expression for what ever use subsequent expressions have.
 	TNiagaraExprPtr MatrixExpr = Expression_Collection(InputExpressions);
-	MatrixExpr->ResultType = ENiagaraDataType::Matrix;
+	MatrixExpr->Result.Type = ENiagaraDataType::Matrix;
 	OutputExpressions.Add(MatrixExpr);
 }
 
diff --git a/Engine/Source/Editor/NiagaraEditor/Private/NiagaraCompiler_VectorVM.h b/Engine/Source/Editor/NiagaraEditor/Private/NiagaraCompiler_VectorVM.h
--- a/Engine/Source/Editor/NiagaraEditor/Private/NiagaraCompiler_VectorVM.h
+++ b/Engine/Source/Editor/NiagaraEditor/Private/NiagaraCompiler_VectorVM.h
@@ -39,11 +39,11 @@ public:
 		NiagaraOpList;
 #undef NiagaraOp
 
-	virtual TNiagaraExprPtr Output(FName OutputName, TNiagaraExprPtr& SourceExpression)override;
+	virtual TNiagaraExprPtr Output(const FNiagaraVariableInfo& Attr, TNiagaraExprPtr& SourceExpression)override;
 	//End INiagaraCompiler;
 
 	//Begin FNiagaraCompiler Interface
-	virtual ENiagaraDataType GetConstantResultIndex(FName Name, bool bInternal, int32& OutResultIndex, int32& OutComponentIndex)override;
+	virtual ENiagaraDataType GetConstantResultIndex(const FNiagaraVariableInfo& Constant, bool bInternal, int32& OutResultIndex, int32& OutComponentIndex)override;
 	virtual int32 AquireTemporary()override;
 	virtual void FreeTemporary(int32 TempIndex) override;
 	//End INiagaraCompiler Interface
diff --git a/Engine/Source/Editor/NiagaraEditor/Private/NiagaraEditor.cpp b/Engine/Source/Editor/NiagaraEditor/Private/NiagaraEditor.cpp
--- a/Engine/Source/Editor/NiagaraEditor/Private/NiagaraEditor.cpp
+++ b/Engine/Source/Editor/NiagaraEditor/Private/NiagaraEditor.cpp
@@ -13,7 +13,7 @@
  
 #define LOCTEXT_NAMESPACE "NiagaraEditor"
 
-const FName FNiagaraEditor::UpdateGraphTabId(TEXT("NiagaraEditor_UpdateGraph"));
+const FName FNiagaraEditor::NodeGraphTabId(TEXT("NiagaraEditor_NodeGraph"));
 const FName FNiagaraEditor::PropertiesTabId(TEXT("NiagaraEditor_MaterialProperties"));
 
 void FNiagaraEditor::RegisterTabSpawners(const TSharedRef<class FTabManager>& TabManager)
@@ -24,8 +24,8 @@ void FNiagaraEditor::RegisterTabSpawners(const TSharedRef<class FTabManager>& Ta
 	
 	TSharedRef<FWorkspaceItem> WorkspaceMenuCategoryRef = WorkspaceMenuCategory.ToSharedRef();
 
-	TabManager->RegisterTabSpawner( UpdateGraphTabId, FOnSpawnTab::CreateSP(this, &FNiagaraEditor::SpawnTab_UpdateGraph) )
-		.SetDisplayName( LOCTEXT("UpdateGraph", "Update Graph") )
+	TabManager->RegisterTabSpawner( NodeGraphTabId, FOnSpawnTab::CreateSP(this, &FNiagaraEditor::SpawnTab_NodeGraph) )
+		.SetDisplayName( LOCTEXT("NodeGraph", "Node Graph") )
 		.SetGroup(WorkspaceMenuCategoryRef);
 
 	TabManager->RegisterTabSpawner(PropertiesTabId, FOnSpawnTab::CreateSP(this, &FNiagaraEditor::SpawnTab_NodeProperties))
@@ -38,7 +38,7 @@ void FNiagaraEditor::UnregisterTabSpawners(const TSharedRef<class FTabManager>&
 {
 	FAssetEditorToolkit::UnregisterTabSpawners(TabManager);
 
-	TabManager->UnregisterTabSpawner(UpdateGraphTabId);
+	TabManager->UnregisterTabSpawner(NodeGraphTabId);
 	TabManager->UnregisterTabSpawner(PropertiesTabId);
 }
 
@@ -53,9 +53,9 @@ void FNiagaraEditor::InitNiagaraEditor( const EToolkitMode::Type Mode, const TSh
 	Script = InScript;
 	check(Script != NULL);
 	Source = CastChecked<UNiagaraScriptSource>(Script->Source);
-	check(Source->UpdateGraph != NULL);
+	check(Source->NodeGraph != NULL);
 
-	TSharedRef<FTabManager::FLayout> StandaloneDefaultLayout = FTabManager::NewLayout( "Standalone_Niagara_Layout_v3" )
+	TSharedRef<FTabManager::FLayout> StandaloneDefaultLayout = FTabManager::NewLayout( "Standalone_Niagara_Layout_v4" )
 	->AddArea
 	(
 		FTabManager::NewPrimaryArea() ->SetOrientation(Orient_Vertical)
@@ -79,7 +79,7 @@ void FNiagaraEditor::InitNiagaraEditor( const EToolkitMode::Type Mode, const TSh
 			(
 				FTabManager::NewStack()
 				->SetSizeCoefficient(0.8f)
-				->AddTab(UpdateGraphTabId, ETabState::OpenedTab)
+				->AddTab(NodeGraphTabId, ETabState::OpenedTab)
 			)
 		)
 	);
@@ -103,7 +103,7 @@ void FNiagaraEditor::InitNiagaraEditor( const EToolkitMode::Type Mode, const TSh
 	if( IsWorldCentricAssetEditor() )
 	{
 		const FString TabInitializationPayload(TEXT(""));		// NOTE: Payload not currently used for table properties
-		SpawnToolkitTab( UpdateGraphTabId, TabInitializationPayload, EToolkitTabSpot::Details );
+		SpawnToolkitTab( NodeGraphTabId, TabInitializationPayload, EToolkitTabSpot::Details );
 	}*/
 }
 
@@ -164,7 +164,7 @@ TSharedRef<SGraphEditor> FNiagaraEditor::CreateGraphEditorWidget(UEdGraph* InGra
 			.FillWidth(1.f)
 			[
 				SNew(STextBlock)
-				.Text(LOCTEXT("UpdateGraphLabel", "Update Graph"))
+				.Text(LOCTEXT("NodeGraphLabel", "Node Graph"))
 				.TextStyle( FEditorStyle::Get(), TEXT("GraphBreadcrumbButtonText") )
 			]
 		];
@@ -179,19 +179,19 @@ TSharedRef<SGraphEditor> FNiagaraEditor::CreateGraphEditorWidget(UEdGraph* InGra
 }
 
 
-TSharedRef<SDockTab> FNiagaraEditor::SpawnTab_UpdateGraph( const FSpawnTabArgs& Args )
+TSharedRef<SDockTab> FNiagaraEditor::SpawnTab_NodeGraph( const FSpawnTabArgs& Args )
 {
-	check( Args.GetTabId().TabType == UpdateGraphTabId );
+	check( Args.GetTabId().TabType == NodeGraphTabId );
 
-	TSharedRef<SGraphEditor> UpdateGraphEditor = CreateGraphEditorWidget(Source->UpdateGraph);
+	TSharedRef<SGraphEditor> NodeGraphEditor = CreateGraphEditorWidget(Source->NodeGraph);
 
-	UpdateGraphEditorPtr = UpdateGraphEditor; // Keep pointer to editor
+	NodeGraphEditorPtr = NodeGraphEditor; // Keep pointer to editor
 
 	return SNew(SDockTab)
-		.Label( LOCTEXT("UpdateGraph", "Update Graph") )
+		.Label( LOCTEXT("NodeGraph", "Node Graph") )
 		.TabColorScale( GetTabColorScale() )
 		[
-			UpdateGraphEditor
+			NodeGraphEditor
 		];
 }
 
@@ -204,11 +204,11 @@ TSharedRef<SDockTab> FNiagaraEditor::SpawnTab_NodeProperties(const FSpawnTabArgs
 			NiagaraDetailsView.ToSharedRef()
 		];
 
-	TSharedPtr<SGraphEditor> UpdateGraphEditor = UpdateGraphEditorPtr.Pin();
-	if (UpdateGraphEditor.IsValid())
+	TSharedPtr<SGraphEditor> NodeGraphEditor = NodeGraphEditorPtr.Pin();
+	if (NodeGraphEditor.IsValid())
 	{
 		// Since we're initialising, make sure nothing is selected
-		UpdateGraphEditor->ClearSelectionSet();
+		NodeGraphEditor->ClearSelectionSet();
 	}
 
 	return SpawnedTab;
@@ -276,14 +276,14 @@ FReply FNiagaraEditor::OnCompileClicked()
 
 void FNiagaraEditor::DeleteSelectedNodes()
 {
-	TSharedPtr<SGraphEditor> UpdateGraphEditor = UpdateGraphEditorPtr.Pin();
-	if (!UpdateGraphEditor.IsValid())
+	TSharedPtr<SGraphEditor> NodeGraphEditor = NodeGraphEditorPtr.Pin();
+	if (!NodeGraphEditor.IsValid())
 	{
 		return;
 	}
 
-	const FGraphPanelSelectionSet SelectedNodes = UpdateGraphEditor->GetSelectedNodes();
-	UpdateGraphEditor->ClearSelectionSet();
+	const FGraphPanelSelectionSet SelectedNodes = NodeGraphEditor->GetSelectedNodes();
+	NodeGraphEditor->ClearSelectionSet();
 
 	for (FGraphPanelSelectionSet::TConstIterator NodeIt( SelectedNodes ); NodeIt; ++NodeIt)
 	{
@@ -299,8 +299,8 @@ void FNiagaraEditor::DeleteSelectedNodes()
 
 bool FNiagaraEditor::CanDeleteNodes() const
 {
-	TSharedPtr<SGraphEditor> UpdateGraphEditor = UpdateGraphEditorPtr.Pin();
-	return (UpdateGraphEditor.IsValid() && UpdateGraphEditor->GetSelectedNodes().Num() > 0);
+	TSharedPtr<SGraphEditor> NodeGraphEditor = NodeGraphEditorPtr.Pin();
+	return (NodeGraphEditor.IsValid() && NodeGraphEditor->GetSelectedNodes().Num() > 0);
 }
 
 void FNiagaraEditor::OnSelectedNodesChanged(const TSet<class UObject*>& NewSelection)
diff --git a/Engine/Source/Editor/NiagaraEditor/Private/NiagaraEditor.h b/Engine/Source/Editor/NiagaraEditor/Private/NiagaraEditor.h
--- a/Engine/Source/Editor/NiagaraEditor/Private/NiagaraEditor.h
+++ b/Engine/Source/Editor/NiagaraEditor/Private/NiagaraEditor.h
@@ -45,7 +45,7 @@ private:
 	TSharedRef<class SGraphEditor> CreateGraphEditorWidget(UEdGraph* InGraph);
 
 	/** Spawns the tab with the update graph inside */
-	TSharedRef<SDockTab> SpawnTab_UpdateGraph(const FSpawnTabArgs& Args);
+	TSharedRef<SDockTab> SpawnTab_NodeGraph(const FSpawnTabArgs& Args);
 
 	/** Spawns the tab with the node details inside. */
 	TSharedRef<SDockTab> SpawnTab_NodeProperties(const FSpawnTabArgs& Args);
@@ -59,7 +59,7 @@ private:
 	/** Source of script being edited */
 	UNiagaraScriptSource*		Source;
 	/** */
-	TWeakPtr<SGraphEditor>		UpdateGraphEditorPtr;
+	TWeakPtr<SGraphEditor>		NodeGraphEditorPtr;
 
 	/** The command list for this editor */
 	TSharedPtr<FUICommandList> GraphEditorCommands;
@@ -68,7 +68,7 @@ private:
 	TSharedPtr<class IDetailsView> NiagaraDetailsView;
 
 	/**	The tab ids for the Niagara editor */
-	static const FName UpdateGraphTabId;
+	static const FName NodeGraphTabId;
 	static const FName PropertiesTabId;
 
 };
diff --git a/Engine/Source/Editor/NiagaraEditor/Private/SNiagaraEffectEditorWidget.cpp b/Engine/Source/Editor/NiagaraEditor/Private/SNiagaraEffectEditorWidget.cpp
--- a/Engine/Source/Editor/NiagaraEditor/Private/SNiagaraEffectEditorWidget.cpp
+++ b/Engine/Source/Editor/NiagaraEditor/Private/SNiagaraEffectEditorWidget.cpp
@@ -556,7 +556,7 @@ void SEmitterWidget::OnUpdateScriptSelectedFromPicker(UObject *Asset)
 
 	// Set the object found from the asset picker
 	CurUpdateScript = Cast<UNiagaraScript>(Asset);
-	Emitter->GetProperties()->UpdateScript = CurUpdateScript;
+	Emitter->SetUpdateScript(CurUpdateScript);
 
 	UpdateScriptConstantListSlot->DetachWidget();
 	(*UpdateScriptConstantListSlot)
@@ -577,7 +577,8 @@ void SEmitterWidget::OnSpawnScriptSelectedFromPicker(UObject *Asset)
 
 	// Set the object found from the asset picker
 	CurSpawnScript = Cast<UNiagaraScript>(Asset);
-	Emitter->GetProperties()->SpawnScript = CurSpawnScript;
+	Emitter->SetSpawnScript(CurSpawnScript);
+
 	SpawnScriptConstantListSlot->DetachWidget();
 	(*SpawnScriptConstantListSlot)
 		[
diff --git a/Engine/Source/Editor/NiagaraEditor/Private/SNiagaraEffectEditorWidget.h b/Engine/Source/Editor/NiagaraEditor/Private/SNiagaraEffectEditorWidget.h
--- a/Engine/Source/Editor/NiagaraEditor/Private/SNiagaraEffectEditorWidget.h
+++ b/Engine/Source/Editor/NiagaraEditor/Private/SNiagaraEffectEditorWidget.h
@@ -47,7 +47,7 @@ public:
 
 	TOptional<float> GetConstX() const	
 	{
-		FVector4 *ConstPtr = Emitter->GetProperties()->ExternalConstants.FindVector(ConstantName);
+		const FVector4 *ConstPtr = Emitter->GetProperties()->ExternalConstants.FindVector(ConstantName);
 		if (ConstPtr)
 		{
 			return ConstPtr->X;
@@ -57,7 +57,7 @@ public:
 
 	TOptional<float> GetConstY() const
 	{
-		FVector4 *ConstPtr = Emitter->GetProperties()->ExternalConstants.FindVector(ConstantName);
+		const FVector4 *ConstPtr = Emitter->GetProperties()->ExternalConstants.FindVector(ConstantName);
 		if (ConstPtr)
 		{
 			return ConstPtr->Y;
@@ -67,7 +67,7 @@ public:
 
 	TOptional<float> GetConstZ() const
 	{
-		FVector4 *ConstPtr = Emitter->GetProperties()->ExternalConstants.FindVector(ConstantName);
+		const FVector4 *ConstPtr = Emitter->GetProperties()->ExternalConstants.FindVector(ConstantName);
 		if (ConstPtr)
 		{
 			return ConstPtr->Z;
@@ -77,7 +77,7 @@ public:
 
 	TOptional<float> GetConstW() const
 	{
-		FVector4 *ConstPtr = Emitter->GetProperties()->ExternalConstants.FindVector(ConstantName);
+		const FVector4 *ConstPtr = Emitter->GetProperties()->ExternalConstants.FindVector(ConstantName);
 		if (ConstPtr)
 		{
 			return ConstPtr->W;
@@ -119,7 +119,7 @@ public:
 		Emitter = InArgs._Emitter;
 		ConstantName = InArgs._ConstantName;
 
-		FVector4 *ConstPtr = Emitter->GetProperties()->ExternalConstants.FindVector(ConstantName);
+		const FVector4 *ConstPtr = Emitter->GetProperties()->ExternalConstants.FindVector(ConstantName);
 		FVector4 ConstVal;
 
 		if (ConstPtr)
@@ -330,8 +330,8 @@ public:
 
 	void OnEmitterEnabledChanged(ECheckBoxState NewCheckedState)
 	{
-		const bool bNewEnabledState = (NewCheckedState == ECheckBoxState::Checked);
-		Emitter->GetProperties()->bIsEnabled = bNewEnabledState;
+		const bool bNewEnabledState = (NewCheckedState == ESlateCheckBoxState::Checked);
+		Emitter->SetEnabled(bNewEnabledState);
 	}
 
 	ECheckBoxState IsEmitterEnabled() const
diff --git a/Engine/Source/Editor/UnrealEd/Classes/Niagara/EdGraphSchema_Niagara.h b/Engine/Source/Editor/UnrealEd/Classes/Niagara/EdGraphSchema_Niagara.h
--- a/Engine/Source/Editor/UnrealEd/Classes/Niagara/EdGraphSchema_Niagara.h
+++ b/Engine/Source/Editor/UnrealEd/Classes/Niagara/EdGraphSchema_Niagara.h
@@ -64,9 +64,12 @@ class UEdGraphSchema_Niagara : public UEdGraphSchema
 	virtual bool ShouldHidePinDefaultValue(UEdGraphPin* Pin) const override;
 	// End EdGraphSchema interface
 
-	UNREALED_API ENiagaraDataType GetPinDataType(UEdGraphPin* Pin)const;
+	UNREALED_API ENiagaraDataType GetPinType(UEdGraphPin* Pin)const;
 	UNREALED_API void GetPinDefaultValue(UEdGraphPin* Pin, float& OutDefault)const;
 	UNREALED_API void GetPinDefaultValue(UEdGraphPin* Pin, FVector4& OutDefault)const;
 	UNREALED_API void GetPinDefaultValue(UEdGraphPin* Pin, FMatrix& OutDefault)const;
+
+	static const FLinearColor NodeTitleColor_Attribute;
+	static const FLinearColor NodeTitleColor_Constant;
 };
 
diff --git a/Engine/Source/Editor/UnrealEd/Classes/Niagara/NiagaraGraph.h b/Engine/Source/Editor/UnrealEd/Classes/Niagara/NiagaraGraph.h
--- a/Engine/Source/Editor/UnrealEd/Classes/Niagara/NiagaraGraph.h
+++ b/Engine/Source/Editor/UnrealEd/Classes/Niagara/NiagaraGraph.h
@@ -8,8 +8,15 @@ UCLASS(MinimalAPI)
 class UNiagaraGraph : public UEdGraph
 {
 	GENERATED_UCLASS_BODY()
-
+	
 	/** Get the source that owns this graph */
 	class UNiagaraScriptSource* GetSource() const;
+		
+	UNREALED_API class UNiagaraNodeOutput* FindOutputNode() const;
+	UNREALED_API void FindInputNodes(TArray<class UNiagaraNodeInput*>& OutInputNodes) const;
+
+	/** Returns the index of this attribute in the output node of the graph. INDEX_NONE if this is not a valid attribute. */
+	UNREALED_API int32 GetAttributeIndex(const FNiagaraVariableInfo& Attr)const;
+	void UNREALED_API GetAttributes(TArray< FNiagaraVariableInfo >& OutAttributes)const;
 };
 
diff --git a/Engine/Source/Editor/UnrealEd/Classes/Niagara/NiagaraNode.h b/Engine/Source/Editor/UnrealEd/Classes/Niagara/NiagaraNode.h
--- a/Engine/Source/Editor/UnrealEd/Classes/Niagara/NiagaraNode.h
+++ b/Engine/Source/Editor/UnrealEd/Classes/Niagara/NiagaraNode.h
@@ -11,10 +11,17 @@ class UNREALED_API UNiagaraNode : public UEdGraphNode
 public:
 
 	/** Get the Niagara graph that owns this node */
+	const class UNiagaraGraph* GetNiagaraGraph()const;
 	class UNiagaraGraph* GetNiagaraGraph();
 
 	/** Get the source object */
-	class UNiagaraScriptSource* GetSource();
+	class UNiagaraScriptSource* GetSource()const;
 
 	virtual void Compile(class INiagaraCompiler* Compiler, TArray<FNiagaraNodeResult>& Outputs){ check(0); }
+
+	UEdGraphPin* GetInputPin(int32 InputIndex) const;
+	void GetInputPins(TArray<class UEdGraphPin*>& OutInputPins) const;
+	UEdGraphPin* GetOutputPin(int32 OutputIndex) const;
+	void GetOutputPins(TArray<class UEdGraphPin*>& OutOutputPins) const;
+
 };
diff --git a/Engine/Source/Editor/UnrealEd/Classes/Niagara/NiagaraNodeConstant.h b/Engine/Source/Editor/UnrealEd/Classes/Niagara/NiagaraNodeConstant.h
--- a/Engine/Source/Editor/UnrealEd/Classes/Niagara/NiagaraNodeConstant.h
+++ b/Engine/Source/Editor/UnrealEd/Classes/Niagara/NiagaraNodeConstant.h
@@ -5,8 +5,8 @@
 #include "Niagara/NiagaraCommon.h"
 #include "NiagaraNodeConstant.generated.h"
 
-UCLASS(MinimalAPI)
-class UNiagaraNodeConstant : public UNiagaraNode
+UCLASS(MinimalAPI, Deprecated)
+class UDEPRECATED_NiagaraNodeConstant : public UNiagaraNode
 {
 	GENERATED_UCLASS_BODY()
 
diff --git a/Engine/Source/Editor/UnrealEd/Classes/Niagara/NiagaraNodeFunctionCall.h b/Engine/Source/Editor/UnrealEd/Classes/Niagara/NiagaraNodeFunctionCall.h
new file mode 100644
--- /dev/null
+++ b/Engine/Source/Editor/UnrealEd/Classes/Niagara/NiagaraNodeFunctionCall.h
@@ -0,0 +1,34 @@
+#pragma once
+
+#include "NiagaraNodeFunctionCall.generated.h"
+
+UCLASS()
+class UNiagaraNodeFunctionCall : public UNiagaraNode
+{
+	GENERATED_BODY()
+
+	UPROPERTY(EditAnywhere, Category = "Function")
+	UNiagaraScript* FunctionScript;
+
+
+	//Begin UObject interface
+	virtual void PostLoad()override;
+	//End UObject interface
+
+	// Begin UNiagaraNode interface
+	virtual void Compile(class INiagaraCompiler* Compiler, TArray<FNiagaraNodeResult>& Outputs)override;
+	//End UNiagaraNode interface
+
+	// Begin EdGraphNode interface
+	virtual void AllocateDefaultPins() override;
+// 	virtual FText GetNodeTitle(ENodeTitleType::Type TitleType) const override;
+// 	virtual FText GetTooltipText() const override;
+// 	virtual FLinearColor GetNodeTitleColor() const override;
+	// End EdGraphNode interface
+
+private:
+	
+	/** Destroys and reallocates input and output pins. */
+	void ReallocatePins();
+};
+
diff --git a/Engine/Source/Editor/UnrealEd/Classes/Niagara/NiagaraNodeGetAttr.h b/Engine/Source/Editor/UnrealEd/Classes/Niagara/NiagaraNodeGetAttr.h
--- a/Engine/Source/Editor/UnrealEd/Classes/Niagara/NiagaraNodeGetAttr.h
+++ b/Engine/Source/Editor/UnrealEd/Classes/Niagara/NiagaraNodeGetAttr.h
@@ -3,8 +3,8 @@
 #pragma once
 #include "NiagaraNodeGetAttr.generated.h"
 
-UCLASS(MinimalAPI)
-class UNiagaraNodeGetAttr : public UNiagaraNode
+UCLASS(MinimalAPI, Deprecated)
+class UDEPRECATED_NiagaraNodeGetAttr : public UNiagaraNode
 {
 	GENERATED_UCLASS_BODY()
 
diff --git a/Engine/Source/Editor/UnrealEd/Classes/Niagara/NiagaraNodeInput.h b/Engine/Source/Editor/UnrealEd/Classes/Niagara/NiagaraNodeInput.h
new file mode 100644
--- /dev/null
+++ b/Engine/Source/Editor/UnrealEd/Classes/Niagara/NiagaraNodeInput.h
@@ -0,0 +1,55 @@
+// Copyright 1998-2014 Epic Games, Inc. All Rights Reserved.
+
+#pragma once
+#include "NiagaraNodeInput.generated.h"
+
+UCLASS(MinimalAPI)
+class UNiagaraNodeInput : public UNiagaraNode
+{
+	GENERATED_UCLASS_BODY()
+
+	void ReallocatePins();
+
+public:
+
+	UPROPERTY(EditAnywhere, Category = Input)
+	FNiagaraVariableInfo Input;
+
+	//TODO: Customize the details for this and hide these when they're not relevant.
+
+	UPROPERTY(EditAnywhere, Category = Constant)
+	float FloatDefault;
+
+	UPROPERTY(EditAnywhere, Category = Constant)
+	FVector4 VectorDefault;
+
+	UPROPERTY(EditAnywhere, Category = Constant)
+	FMatrix MatrixDefault;
+
+	/** Allows code to explicitly disable exposing of certain inputs e.g. system constants such as Delta Time. */
+	UPROPERTY()
+	uint32 bCanBeExposed:1;
+
+	/** When true, and this input is a constant, the input is exposed to the effect editor. */
+	UPROPERTY(EditAnywhere, Category = Constant, meta = (editcondition = "bCanBeExposed"))
+	uint32 bExposeWhenConstant:1;
+
+	// Begin UObject interface
+	virtual void PostEditChangeProperty(struct FPropertyChangedEvent& PropertyChangedEvent);
+	// End UObject interface
+
+	// Begin EdGraphNode interface
+	virtual void AllocateDefaultPins() override;
+	virtual FText GetNodeTitle(ENodeTitleType::Type TitleType) const override;
+	virtual FLinearColor GetNodeTitleColor() const override;
+	// End EdGraphNode interface
+
+	virtual void Compile(class INiagaraCompiler* Compiler, TArray<FNiagaraNodeResult>& Outputs)override;
+	
+	bool IsConstant()const;
+	bool IsExposedConstant()const;
+
+	static const FLinearColor TitleColor_Attribute;
+	static const FLinearColor TitleColor_Constant;
+};
+
diff --git a/Engine/Source/Editor/UnrealEd/Classes/Niagara/NiagaraNodeOutput.h b/Engine/Source/Editor/UnrealEd/Classes/Niagara/NiagaraNodeOutput.h
new file mode 100644
--- /dev/null
+++ b/Engine/Source/Editor/UnrealEd/Classes/Niagara/NiagaraNodeOutput.h
@@ -0,0 +1,30 @@
+// Copyright 1998-2014 Epic Games, Inc. All Rights Reserved.
+
+#pragma once
+#include "NiagaraNodeOutput.generated.h"
+
+UCLASS(MinimalAPI)
+class UNiagaraNodeOutput : public UNiagaraNode
+{
+	GENERATED_UCLASS_BODY()
+
+	void ReallocatePins();
+public:
+
+	UPROPERTY(EditAnywhere, Category = Output)
+	TArray<FNiagaraVariableInfo> Outputs;
+
+	// Begin UObject interface
+	virtual void PostEditChangeProperty(struct FPropertyChangedEvent& PropertyChangedEvent);
+	// End UObject interface
+
+	// Begin EdGraphNode interface
+	virtual void AllocateDefaultPins() override;
+	virtual bool CanUserDeleteNode() const override { return false; }
+	virtual FText GetNodeTitle(ENodeTitleType::Type TitleType) const override;
+	virtual FLinearColor GetNodeTitleColor() const override;
+	// End EdGraphNode interface
+
+	virtual void Compile(class INiagaraCompiler* Compiler, TArray<FNiagaraNodeResult>& OutputExpressions)override;
+};
+
diff --git a/Engine/Source/Editor/UnrealEd/Classes/Niagara/NiagaraNodeOutputUpdate.h b/Engine/Source/Editor/UnrealEd/Classes/Niagara/NiagaraNodeOutputUpdate.h
--- a/Engine/Source/Editor/UnrealEd/Classes/Niagara/NiagaraNodeOutputUpdate.h
+++ b/Engine/Source/Editor/UnrealEd/Classes/Niagara/NiagaraNodeOutputUpdate.h
@@ -3,8 +3,8 @@
 #pragma once
 #include "NiagaraNodeOutputUpdate.generated.h"
 
-UCLASS(MinimalAPI)
-class UNiagaraNodeOutputUpdate : public UNiagaraNode
+UCLASS(MinimalAPI, Deprecated)
+class UDEPRECATED_NiagaraNodeOutputUpdate : public UNiagaraNode
 {
 	GENERATED_UCLASS_BODY()
 
diff --git a/Engine/Source/Editor/UnrealEd/Classes/Niagara/NiagaraScriptSource.h b/Engine/Source/Editor/UnrealEd/Classes/Niagara/NiagaraScriptSource.h
--- a/Engine/Source/Editor/UnrealEd/Classes/Niagara/NiagaraScriptSource.h
+++ b/Engine/Source/Editor/UnrealEd/Classes/Niagara/NiagaraScriptSource.h
@@ -13,12 +13,11 @@ class UNiagaraScriptSource : public UNiagaraScriptSourceBase
 
 	/** Graph for particle update expression */
 	UPROPERTY()
-	class UNiagaraGraph*	UpdateGraph;
+	class UNiagaraGraph*	NodeGraph;
 
 	// UObject interface.
 	virtual void PostLoad() override;
 	virtual void Compile() override;
 
-	UNREALED_API void GetParticleAttributes(TArray<FName>& VectorOutputs);
 	UNREALED_API void GetEmitterAttributes(TArray<FName>& VectorInputs, TArray<FName>& MatrixInputs);
 };
diff --git a/Engine/Source/Editor/UnrealEd/Private/Factories/NiagaraScriptFactory.cpp b/Engine/Source/Editor/UnrealEd/Private/Factories/NiagaraScriptFactory.cpp
--- a/Engine/Source/Editor/UnrealEd/Private/Factories/NiagaraScriptFactory.cpp
+++ b/Engine/Source/Editor/UnrealEd/Private/Factories/NiagaraScriptFactory.cpp
@@ -34,91 +34,18 @@ UObject* UNiagaraScriptFactoryNew::FactoryCreateNew(UClass* Class,UObject* InPar
 		{
 			// Create 'update graph'
 			UNiagaraGraph* CreatedGraph  = ConstructObject<UNiagaraGraph>(UNiagaraGraph::StaticClass(), Source, NAME_None, RF_Transactional);
-			Source->UpdateGraph = CreatedGraph;
+			Source->NodeGraph = CreatedGraph;
 
-			FGraphNodeCreator<UNiagaraNodeOutputUpdate> OutputNodeCreator(*CreatedGraph);
-			UNiagaraNodeOutputUpdate* OutputNode = OutputNodeCreator.CreateNode();
-			OutputNodeCreator.Finalize();
-
-// 			FGraphNodeCreator<UNiagaraNodeGetAttr> TimeNodeCreator(*CreatedGraph);
-// 			UNiagaraNodeGetAttr* DeltaTimeNode = TimeNodeCreator.CreateNode();
-// 			DeltaTimeNode->AttrName = FName(TEXT("DeltaTime"));
-// 			DeltaTimeNode->bIsEmitterAttribute = true;
-// 			TimeNodeCreator.Finalize();
-			//UEdGraphPin* DeltaTimePin = DeltaTimeNode->FindPinChecked(DeltaTimeNode->AttrName.ToString());
-
-			TArray<FName> OutputNames;
-			Source->GetParticleAttributes(OutputNames);
-
-			// Simple Euler integration.
-			/*
-			for(int32 i=0; i<3; i++)
-			{
-				FGraphNodeCreator<UNiagaraNodeGetAttr> PosNodeCreator(*CreatedGraph);
-				UNiagaraNodeGetAttr* PosAttrNode = PosNodeCreator.CreateNode();
-				PosAttrNode->AttrName = OutputNames[i];
-				PosNodeCreator.Finalize();
-
-				FGraphNodeCreator<UNiagaraNodeGetAttr> VelNodeCreator(*CreatedGraph);
-				UNiagaraNodeGetAttr* VelAttrNode = VelNodeCreator.CreateNode();
-				VelAttrNode->AttrName = OutputNames[i+3];
-				VelNodeCreator.Finalize();
-
-
-				UEdGraphPin* PosAttrOutputPin = PosAttrNode->FindPinChecked(OutputNames[i].ToString());
-				UEdGraphPin* VelAttrOutputPin = VelAttrNode->FindPinChecked(OutputNames[i+3].ToString());
-				UEdGraphPin* OutputPin = OutputNode->FindPinChecked(OutputNames[i].ToString());
-
-				UNiagaraNodeOp* MulNode = NULL;
-				{
-					FGraphNodeCreator<UNiagaraNodeOp> NodeCreator(*CreatedGraph);
-					MulNode = NodeCreator.CreateNode();
-					MulNode->OpIndex = VectorVM::EOp::mul;
-					NodeCreator.Finalize();
+			FGraphNodeCreator<UNiagaraNodeOutput> OutputNodeCreator(*CreatedGraph);
+			UNiagaraNodeOutput* OutputNode = OutputNodeCreator.CreateNode();
 
-					UEdGraphPin* APin = MulNode->FindPinChecked(TEXT("A"));
-					UEdGraphPin* BPin = MulNode->FindPinChecked(TEXT("B"));
+			//Automatically add attributes that the current render modules rely on. Not sure what we'll want to do here in the end.
+			OutputNode->Outputs.Add(FNiagaraVariableInfo(FName(TEXT("Position")), ENiagaraDataType::Vector));
+			OutputNode->Outputs.Add(FNiagaraVariableInfo(FName(TEXT("Rotation")), ENiagaraDataType::Vector));
+			OutputNode->Outputs.Add(FNiagaraVariableInfo(FName(TEXT("Color")), ENiagaraDataType::Vector));
+			OutputNode->Outputs.Add(FNiagaraVariableInfo(FName(TEXT("Age")), ENiagaraDataType::Vector));
 
-					APin->MakeLinkTo(VelAttrOutputPin);
-					BPin->MakeLinkTo(DeltaTimePin);
-				}
-
-				{
-					FGraphNodeCreator<UNiagaraNodeOp> NodeCreator(*CreatedGraph);
-					UNiagaraNodeOp* AddNode = NodeCreator.CreateNode();
-					AddNode->OpIndex = VectorVM::EOp::add;
-					NodeCreator.Finalize();
-
-					UEdGraphPin* APin = AddNode->FindPinChecked(TEXT("A"));
-					UEdGraphPin* BPin = AddNode->FindPinChecked(TEXT("B"));
-					UEdGraphPin* AddResultPin = AddNode->FindPinChecked(TEXT("Result"));
-					UEdGraphPin* MulResultPin = MulNode->FindPinChecked(TEXT("Result"));
-
-					APin->MakeLinkTo(PosAttrOutputPin);
-					BPin->MakeLinkTo(MulResultPin);
-					OutputPin->MakeLinkTo(AddResultPin);
-				}
-			}
-
-			// Update relative time.
-			{
-				FGraphNodeCreator<UNiagaraNodeGetAttr> RelTimeNodeCreator(*Source->UpdateGraph);
-				UNiagaraNodeGetAttr* RelTimeAttrNode = RelTimeNodeCreator.CreateNode();
-				RelTimeAttrNode->AttrName = FName(TEXT("RelativeTime"));
-				RelTimeNodeCreator.Finalize();
-
-				FGraphNodeCreator<UNiagaraNodeOp> MadNodeCreator(*Source->UpdateGraph);
-				UNiagaraNodeOp* MadNode = MadNodeCreator.CreateNode();
-				MadNode->OpIndex = VectorVM::EOp::mad;
-				MadNodeCreator.Finalize();
-
-				MadNode->FindPinChecked(TEXT("A"))->MakeLinkTo(DeltaTimePin);
-				MadNode->FindPinChecked(TEXT("B"))->DefaultValue = TEXT("0.2");
-				MadNode->FindPinChecked(TEXT("C"))->MakeLinkTo(RelTimeAttrNode->FindPinChecked(TEXT("RelativeTime")));
-
-				OutputNode->FindPinChecked(TEXT("RelativeTime"))->MakeLinkTo(MadNode->FindPinChecked(TEXT("Result")));
-			}
-			*/
+			OutputNodeCreator.Finalize();
 
 			// Set pointer in script to source
 			NewScript->Source = Source;
diff --git a/Engine/Source/Editor/UnrealEd/Private/Niagara/EdGraphSchema_Niagara.cpp b/Engine/Source/Editor/UnrealEd/Private/Niagara/EdGraphSchema_Niagara.cpp
--- a/Engine/Source/Editor/UnrealEd/Private/Niagara/EdGraphSchema_Niagara.cpp
+++ b/Engine/Source/Editor/UnrealEd/Private/Niagara/EdGraphSchema_Niagara.cpp
@@ -4,7 +4,7 @@
 
 #include "BlueprintGraphDefinitions.h"
 #include "GraphEditorSettings.h"
-
+#include "Components/NiagaraComponent.h"
 
 #include "INiagaraEditor.h"
 #include "INiagaraCompiler.h"
@@ -13,6 +13,9 @@
 
 #define SNAP_GRID (16) // @todo ensure this is the same as SNodePanel::GetSnapGridSize()
 
+const FLinearColor UEdGraphSchema_Niagara::NodeTitleColor_Attribute = FLinearColor::Blue;
+const FLinearColor UEdGraphSchema_Niagara::NodeTitleColor_Constant = FLinearColor::Red;
+
 namespace 
 {
 	// Maximum distance a drag can be off a node edge to require 'push off' from node
@@ -115,81 +118,71 @@ TSharedPtr<FNiagaraSchemaAction_NewNode> AddNewNodeAction(FGraphContextMenuBuild
 void UEdGraphSchema_Niagara::GetGraphContextActions(FGraphContextMenuBuilder& ContextMenuBuilder) const
 {
 	const UNiagaraGraph* NiagaraGraph = CastChecked<UNiagaraGraph>(ContextMenuBuilder.CurrentGraph);
-	UNiagaraScriptSource* Source = NiagaraGraph->GetSource();
-
-	TArray<FName> InputAttributeNames;
-	Source->GetParticleAttributes(InputAttributeNames);
-	for(int32 i=0; i<InputAttributeNames.Num(); i++)
-	{
-		const FName AttrName = InputAttributeNames[i];
-
-		FFormatNamedArguments Args;
-		Args.Add(TEXT("Attribute"), FText::FromName(AttrName));
-		const FText MenuDesc = FText::Format(NSLOCTEXT("Niagara", "GetAttribute", "Get {Attribute}"), Args);
-
-		TSharedPtr<FNiagaraSchemaAction_NewNode> GetAttrAction = AddNewNodeAction(ContextMenuBuilder, TEXT("Get Attribute"), MenuDesc, TEXT(""));
-
-		UNiagaraNodeGetAttr* GetAttrNode = NewObject<UNiagaraNodeGetAttr>(ContextMenuBuilder.OwnerOfTemporaries);
-		GetAttrNode->AttrName = AttrName;
-		GetAttrAction->NodeTemplate = GetAttrNode;
-	}
-
-	TArray<FName> EmitterConstantNames_Vectors;
-	TArray<FName> EmitterConstantNames_Matrices;
-	Source->GetEmitterAttributes(EmitterConstantNames_Vectors, EmitterConstantNames_Matrices);
-	for (int32 i = 0; i < EmitterConstantNames_Vectors.Num(); i++)
-	{
-		const FName ConstName = EmitterConstantNames_Vectors[i];
-
-		FFormatNamedArguments Args;
-		Args.Add(TEXT("Constant"), FText::FromName(ConstName));
-		const FText MenuDesc = FText::Format(NSLOCTEXT("Niagara", "GetConstant", "Get {Constant}"), Args);
-
-		TSharedPtr<FNiagaraSchemaAction_NewNode> GetAttrAction = AddNewNodeAction(ContextMenuBuilder, TEXT("Get Constant"), MenuDesc, TEXT(""));
-
-		UNiagaraNodeConstant* GetConstNode = NewObject<UNiagaraNodeConstant>(ContextMenuBuilder.OwnerOfTemporaries);
-		GetConstNode->ConstName = ConstName;
-		GetConstNode->DataType = ENiagaraDataType::Vector;
-		GetConstNode->bNeedsDefault = false;
-		GetAttrAction->NodeTemplate = GetConstNode;
-	}
 
-	for (int32 i = 0; i < EmitterConstantNames_Matrices.Num(); i++)
+	//Add menu options for Attributes currently defined by the output node.
+	const UNiagaraNodeOutput* OutNode = NiagaraGraph->FindOutputNode();
+	if (OutNode)
 	{
-		const FName ConstName = EmitterConstantNames_Matrices[i];
+		for (int32 i = 0; i < OutNode->Outputs.Num(); i++)
+		{
+			const FNiagaraVariableInfo Attr = OutNode->Outputs[i];
 
-		FFormatNamedArguments Args;
-		Args.Add(TEXT("Constant"), FText::FromName(ConstName));
-		const FText MenuDesc = FText::Format(NSLOCTEXT("Niagara", "GetConstant", "Get {Constant}"), Args);
+			FFormatNamedArguments Args;
+			Args.Add(TEXT("Attribute"), FText::FromName(Attr.Name));
+			const FText MenuDesc = FText::Format(LOCTEXT("GetAttribute", "Get {Attribute}"), Args);
 
-		TSharedPtr<FNiagaraSchemaAction_NewNode> GetAttrAction = AddNewNodeAction(ContextMenuBuilder, TEXT("Get Constant"), MenuDesc, TEXT(""));
+			TSharedPtr<FNiagaraSchemaAction_NewNode> GetAttrAction = AddNewNodeAction(ContextMenuBuilder, *LOCTEXT("Attributes Menu Title", "Attributes").ToString(), MenuDesc, TEXT(""));
 
-		UNiagaraNodeConstant* GetConstNode = NewObject<UNiagaraNodeConstant>(ContextMenuBuilder.OwnerOfTemporaries);
-		GetConstNode->ConstName = ConstName;
-		GetConstNode->DataType = ENiagaraDataType::Matrix;
-		GetConstNode->bNeedsDefault = false;
-		GetAttrAction->NodeTemplate = GetConstNode;
+			UNiagaraNodeInput* InputNode = NewObject<UNiagaraNodeInput>(ContextMenuBuilder.OwnerOfTemporaries);
+			InputNode->Input = Attr;
+			GetAttrAction->NodeTemplate = InputNode;
+		}
 	}
 
 #define NiagaraOp(OPNAME) \
 		if(const FNiagaraOpInfo* OpInfo = FNiagaraOpInfo::GetOpInfo(FNiagaraOpInfo::OPNAME))\
 		{\
-			TSharedPtr<FNiagaraSchemaAction_NewNode> AddOpAction = AddNewNodeAction(ContextMenuBuilder, TEXT("Add Operation"), OpInfo->FriendlyName, TEXT(""));\
+		TSharedPtr<FNiagaraSchemaAction_NewNode> AddOpAction = AddNewNodeAction(ContextMenuBuilder, *LOCTEXT("Operations Menu Title", "Operations").ToString(), OpInfo->FriendlyName, TEXT("")); \
 			UNiagaraNodeOp* OpNode = NewObject<UNiagaraNodeOp>(ContextMenuBuilder.OwnerOfTemporaries); \
 			OpNode->OpName = OpInfo->Name; \
 			AddOpAction->NodeTemplate = OpNode; \
 		}
 
-		
 	NiagaraOpList
 
 #undef NiagaraOp
 
-	UNiagaraNodeConstant* ConstantNode = NewObject<UNiagaraNodeConstant>(ContextMenuBuilder.OwnerOfTemporaries);
-	ConstantNode->DataType = ENiagaraDataType::Vector;
-	ConstantNode->bNeedsDefault = true;
-	TSharedPtr<FNiagaraSchemaAction_NewNode> AddOpAction = AddNewNodeAction(ContextMenuBuilder, TEXT("Constants"), ConstantNode->GetNodeTitle(ENodeTitleType::ListView), TEXT(""));
-	AddOpAction->NodeTemplate = ConstantNode;
+	//Emitter constants managed by the system.
+	const TArray<FNiagaraVariableInfo>& SystemConstants = UNiagaraComponent::GetSystemConstants();
+	for (const FNiagaraVariableInfo& SysConst : SystemConstants)
+	{
+		FFormatNamedArguments Args;
+		Args.Add(TEXT("Constant"), FText::FromName(SysConst.Name));
+		const FText MenuDesc = FText::Format(LOCTEXT("GetSystemConstant", "Get {Constant}"), Args);
+
+		TSharedPtr<FNiagaraSchemaAction_NewNode> GetConstAction = AddNewNodeAction(ContextMenuBuilder, *LOCTEXT("System Constants Menu Title", "System Constants").ToString(), MenuDesc, TEXT(""));
+
+		UNiagaraNodeInput* InputNode = NewObject<UNiagaraNodeInput>(ContextMenuBuilder.OwnerOfTemporaries);
+		InputNode->Input = SysConst;
+		InputNode->bExposeWhenConstant = false;
+		InputNode->bCanBeExposed = false;
+		GetConstAction->NodeTemplate = InputNode;
+	}
+
+	//Add a generic Input node to allow getting external constants the current context doesn't know about.
+	{
+		const FText MenuDesc = LOCTEXT("GetSystemConstantUnknown", "Get Constant");
+
+		//TSharedPtr<FNiagaraSchemaAction_NewNode> GetConstAction = AddNewNodeAction(ContextMenuBuilder, *LOCTEXT("System Constants Menu Title", "System Constants").ToString(), MenuDesc, TEXT(""));
+		TSharedPtr<FNiagaraSchemaAction_NewNode> GetConstAction = AddNewNodeAction(ContextMenuBuilder, TEXT(""), MenuDesc, TEXT(""));
+
+		UNiagaraNodeInput* InputNode = NewObject<UNiagaraNodeInput>(ContextMenuBuilder.OwnerOfTemporaries);
+		InputNode->Input.Name = NAME_None;
+		InputNode->Input.Type = ENiagaraDataType::Vector;
+		GetConstAction->NodeTemplate = InputNode;
+	}
+
+	//TODO: Add quick commands for certain UNiagaraStructs and UNiagaraScripts to be added as functions
 }
 
 const FPinConnectionResponse UEdGraphSchema_Niagara::CanCreateConnection(const UEdGraphPin* PinA, const UEdGraphPin* PinB) const
@@ -271,7 +264,7 @@ bool UEdGraphSchema_Niagara::ShouldHidePinDefaultValue(UEdGraphPin* Pin) const
 	return false;
 }
 
-ENiagaraDataType UEdGraphSchema_Niagara::GetPinDataType(UEdGraphPin* Pin) const
+ENiagaraDataType UEdGraphSchema_Niagara::GetPinType(UEdGraphPin* Pin) const
 {
 	if (Pin->PinType.PinCategory == PC_Float)
 	{
@@ -289,7 +282,8 @@ ENiagaraDataType UEdGraphSchema_Niagara::GetPinDataType(UEdGraphPin* Pin) const
 	{
 		return ENiagaraDataType::Curve;
 	}
-	return ENiagaraDataType::Invalid;
+	check(0);
+	return ENiagaraDataType::Vector;
 }
 
 void UEdGraphSchema_Niagara::GetPinDefaultValue(UEdGraphPin* Pin, float& OutDefault)const
diff --git a/Engine/Source/Editor/UnrealEd/Private/Niagara/NiagaraNode.cpp b/Engine/Source/Editor/UnrealEd/Private/Niagara/NiagaraNode.cpp
--- a/Engine/Source/Editor/UnrealEd/Private/Niagara/NiagaraNode.cpp
+++ b/Engine/Source/Editor/UnrealEd/Private/Niagara/NiagaraNode.cpp
@@ -7,12 +7,83 @@ UNiagaraNode::UNiagaraNode(const FObjectInitializer& ObjectInitializer)
 {
 }
 
+const UNiagaraGraph* UNiagaraNode::GetNiagaraGraph()const
+{
+	return CastChecked<UNiagaraGraph>(GetGraph());
+}
+
 UNiagaraGraph* UNiagaraNode::GetNiagaraGraph()
 {
 	return CastChecked<UNiagaraGraph>(GetGraph());
 }
 
-class UNiagaraScriptSource* UNiagaraNode::GetSource()
+UNiagaraScriptSource* UNiagaraNode::GetSource()const
 {
 	return GetNiagaraGraph()->GetSource();
 }
+
+UEdGraphPin* UNiagaraNode::GetInputPin(int32 InputIndex) const
+{
+	for (int32 PinIndex = 0, FoundInputs = 0; PinIndex < Pins.Num(); PinIndex++)
+	{
+		if (Pins[PinIndex]->Direction == EGPD_Input)
+		{
+			if (InputIndex == FoundInputs)
+			{
+				return Pins[PinIndex];
+			}
+			else
+			{
+				FoundInputs++;
+			}
+		}
+	}
+
+	return NULL;
+}
+
+void UNiagaraNode::GetInputPins(TArray<class UEdGraphPin*>& OutInputPins) const
+{
+	OutInputPins.Empty();
+
+	for (int32 PinIndex = 0; PinIndex < Pins.Num(); PinIndex++)
+	{
+		if (Pins[PinIndex]->Direction == EGPD_Input)
+		{
+			OutInputPins.Add(Pins[PinIndex]);
+		}
+	}
+}
+
+UEdGraphPin* UNiagaraNode::GetOutputPin(int32 OutputIndex) const
+{
+	for (int32 PinIndex = 0, FoundOutputs = 0; PinIndex < Pins.Num(); PinIndex++)
+	{
+		if (Pins[PinIndex]->Direction == EGPD_Output)
+		{
+			if (OutputIndex == FoundOutputs)
+			{
+				return Pins[PinIndex];
+			}
+			else
+			{
+				FoundOutputs++;
+			}
+		}
+	}
+
+	return NULL;
+}
+
+void UNiagaraNode::GetOutputPins(TArray<class UEdGraphPin*>& OutOutputPins) const
+{
+	OutOutputPins.Empty();
+
+	for (int32 PinIndex = 0; PinIndex < Pins.Num(); PinIndex++)
+	{
+		if (Pins[PinIndex]->Direction == EGPD_Output)
+		{
+			OutOutputPins.Add(Pins[PinIndex]);
+		}
+	}
+}
diff --git a/Engine/Source/Editor/UnrealEd/Private/Niagara/NiagaraNodeConstant.cpp b/Engine/Source/Editor/UnrealEd/Private/Niagara/NiagaraNodeConstant.cpp
--- a/Engine/Source/Editor/UnrealEd/Private/Niagara/NiagaraNodeConstant.cpp
+++ b/Engine/Source/Editor/UnrealEd/Private/Niagara/NiagaraNodeConstant.cpp
@@ -5,7 +5,7 @@
 
 #define LOCTEXT_NAMESPACE "NiagaraNodeConst"
 
-UNiagaraNodeConstant::UNiagaraNodeConstant(const FObjectInitializer& ObjectInitializer)
+UDEPRECATED_NiagaraNodeConstant::UDEPRECATED_NiagaraNodeConstant(const FObjectInitializer& ObjectInitializer)
 	: Super(ObjectInitializer)
 {
 	ConstName = NAME_None;
@@ -14,7 +14,7 @@ UNiagaraNodeConstant::UNiagaraNodeConstant(const FObjectInitializer& ObjectIniti
 	bExposeToEffectEditor = false;
 }
 
-void UNiagaraNodeConstant::AllocateDefaultPins()
+void UDEPRECATED_NiagaraNodeConstant::AllocateDefaultPins()
 {
 	const UEdGraphSchema_Niagara* Schema = GetDefault<UEdGraphSchema_Niagara>();
 
@@ -52,7 +52,7 @@ void UNiagaraNodeConstant::AllocateDefaultPins()
 	}
 }
 
-FText UNiagaraNodeConstant::GetNodeTitle(ENodeTitleType::Type TitleType) const
+FText UDEPRECATED_NiagaraNodeConstant::GetNodeTitle(ENodeTitleType::Type TitleType) const
 {
 	if (ConstName != NAME_None)
 		return FText::FromName(ConstName);
@@ -60,7 +60,7 @@ FText UNiagaraNodeConstant::GetNodeTitle(ENodeTitleType::Type TitleType) const
 		return LOCTEXT("DefaultTitle", "External Constant");
 }
 
-FLinearColor UNiagaraNodeConstant::GetNodeTitleColor() const
+FLinearColor UDEPRECATED_NiagaraNodeConstant::GetNodeTitleColor() const
 {
 	const UEdGraphSchema_Niagara* Schema = GetDefault<UEdGraphSchema_Niagara>();
 
@@ -74,74 +74,80 @@ FLinearColor UNiagaraNodeConstant::GetNodeTitleColor() const
 	}
 }
 
-void UNiagaraNodeConstant::Compile(class INiagaraCompiler* Compiler, TArray<FNiagaraNodeResult>& Outputs)
+void UDEPRECATED_NiagaraNodeConstant::Compile(class INiagaraCompiler* Compiler, TArray<FNiagaraNodeResult>& Outputs)
 {
-	TArray<int32> OutputExpressions;
-	if (DataType == ENiagaraDataType::Scalar)
-	{
-		float Default = 1.0f;
-		UEdGraphPin* OutPin;
-		if (bNeedsDefault)
-		{
-			CastChecked<UEdGraphSchema_Niagara>(GetSchema())->GetPinDefaultValue(Pins[0], Default);
-			OutPin = Pins[1];
-		}
-		else
-		{
-			OutPin = Pins[0];
-		}
-		Outputs.Add(FNiagaraNodeResult(Compiler->GetExternalConstant(ConstName, Default), OutPin));
-	}
-	else if (DataType == ENiagaraDataType::Vector)
-	{
-		FVector4 Default = FVector4(1.0f, 1.0f, 1.0f, 1.0f);
-		UEdGraphPin* OutPin;
-		if (bNeedsDefault)
-		{
-			CastChecked<UEdGraphSchema_Niagara>(GetSchema())->GetPinDefaultValue(Pins[0], Default);
-			OutPin = Pins[1];
-		}
-		else
-		{
-			OutPin = Pins[0];
-		}
-		Outputs.Add(FNiagaraNodeResult(Compiler->GetExternalConstant(ConstName, Default), OutPin));
-	}
-	else if (DataType == ENiagaraDataType::Matrix)
-	{
-		FVector4 Default0 = FVector4(1.0f, 0.0f, 0.0f, 0.0f);
-		FVector4 Default1 = FVector4(0.0f, 1.0f, 0.0f, 0.0f);
-		FVector4 Default2 = FVector4(0.0f, 0.0f, 1.0f, 0.0f);
-		FVector4 Default3 = FVector4(0.0f, 0.0f, 0.0f, 1.0f);
+	//UE_LOG(LogNiagara, Fatal, TEXT("UDEPRECATED_NiagaraNodeConstant is deprecated. All instance of it should have been replaced in UNiagaraGraph::PostLoad()"));
+	check(0);
 
-		UEdGraphPin* OutPin;
-		if (bNeedsDefault)
-		{
-			CastChecked<UEdGraphSchema_Niagara>(GetSchema())->GetPinDefaultValue(Pins[0], Default0);
-			CastChecked<UEdGraphSchema_Niagara>(GetSchema())->GetPinDefaultValue(Pins[1], Default1);
-			CastChecked<UEdGraphSchema_Niagara>(GetSchema())->GetPinDefaultValue(Pins[2], Default2);
-			CastChecked<UEdGraphSchema_Niagara>(GetSchema())->GetPinDefaultValue(Pins[3], Default3);
-			OutPin = Pins[4];
-		}
-		else
-		{
-			OutPin = Pins[0];
-		}
-		FMatrix Default;
-		Default.M[0][0] = Default0[0]; Default.M[0][1] = Default0[1]; Default.M[0][2] = Default0[2]; Default.M[0][3] = Default0[3];
-		Default.M[1][0] = Default1[0]; Default.M[1][1] = Default1[1]; Default.M[1][2] = Default1[2]; Default.M[1][3] = Default1[3];
-		Default.M[2][0] = Default2[0]; Default.M[2][1] = Default2[1]; Default.M[2][2] = Default2[2]; Default.M[2][3] = Default2[3];
-		Default.M[3][0] = Default3[0]; Default.M[3][1] = Default3[1]; Default.M[3][2] = Default3[2]; Default.M[3][3] = Default3[3];
-		Outputs.Add(FNiagaraNodeResult(Compiler->GetExternalConstant(ConstName, Default), OutPin));
-	}
-	else if (DataType == ENiagaraDataType::Curve)
-	{
-		UEdGraphPin* OutPin = Pins[0];
-		Outputs.Add(FNiagaraNodeResult(Compiler->GetExternalCurveConstant(ConstName), OutPin));
-	}
+	return;
+
+	//TODO REMOVE THIS NODE ENTIRELY
+// 	TArray<int32> OutputExpressions;
+// 	if (DataType == ENiagaraDataType::Scalar)
+// 	{
+// 		float Default = 1.0f;
+// 		UEdGraphPin* OutPin;
+// 		if (bNeedsDefault)
+// 		{
+// 			CastChecked<UEdGraphSchema_Niagara>(GetSchema())->GetPinDefaultValue(Pins[0], Default);
+// 			OutPin = Pins[1];
+// 		}
+// 		else
+// 		{
+// 			OutPin = Pins[0];
+// 		}
+// 		Outputs.Add(FNiagaraNodeResult(Compiler->GetExternalConstant(ConstName, Default), OutPin));
+// 	}
+// 	else if (DataType == ENiagaraDataType::Vector)
+// 	{
+// 		FVector4 Default = FVector4(1.0f, 1.0f, 1.0f, 1.0f);
+// 		UEdGraphPin* OutPin;
+// 		if (bNeedsDefault)
+// 		{
+// 			CastChecked<UEdGraphSchema_Niagara>(GetSchema())->GetPinDefaultValue(Pins[0], Default);
+// 			OutPin = Pins[1];
+// 		}
+// 		else
+// 		{
+// 			OutPin = Pins[0];
+// 		}
+// 		Outputs.Add(FNiagaraNodeResult(Compiler->GetExternalConstant(ConstName, Default), OutPin));
+// 	}
+// 	else if (DataType == ENiagaraDataType::Matrix)
+// 	{
+// 		FVector4 Default0 = FVector4(1.0f, 0.0f, 0.0f, 0.0f);
+// 		FVector4 Default1 = FVector4(0.0f, 1.0f, 0.0f, 0.0f);
+// 		FVector4 Default2 = FVector4(0.0f, 0.0f, 1.0f, 0.0f);
+// 		FVector4 Default3 = FVector4(0.0f, 0.0f, 0.0f, 1.0f);
+// 
+// 		UEdGraphPin* OutPin;
+// 		if (bNeedsDefault)
+// 		{
+// 			CastChecked<UEdGraphSchema_Niagara>(GetSchema())->GetPinDefaultValue(Pins[0], Default0);
+// 			CastChecked<UEdGraphSchema_Niagara>(GetSchema())->GetPinDefaultValue(Pins[1], Default1);
+// 			CastChecked<UEdGraphSchema_Niagara>(GetSchema())->GetPinDefaultValue(Pins[2], Default2);
+// 			CastChecked<UEdGraphSchema_Niagara>(GetSchema())->GetPinDefaultValue(Pins[3], Default3);
+// 			OutPin = Pins[4];
+// 		}
+// 		else
+// 		{
+// 			OutPin = Pins[0];
+// 		}
+// 		FMatrix Default;
+// 		Default.M[0][0] = Default0[0]; Default.M[0][1] = Default0[1]; Default.M[0][2] = Default0[2]; Default.M[0][3] = Default0[3];
+// 		Default.M[1][0] = Default1[0]; Default.M[1][1] = Default1[1]; Default.M[1][2] = Default1[2]; Default.M[1][3] = Default1[3];
+// 		Default.M[2][0] = Default2[0]; Default.M[2][1] = Default2[1]; Default.M[2][2] = Default2[2]; Default.M[2][3] = Default2[3];
+// 		Default.M[3][0] = Default3[0]; Default.M[3][1] = Default3[1]; Default.M[3][2] = Default3[2]; Default.M[3][3] = Default3[3];
+// 		Outputs.Add(FNiagaraNodeResult(Compiler->GetExternalConstant(ConstName, Default), OutPin));
+// 	}
+//	else if (DataType == ENiagaraDataType::Curve)
+//	{
+//		UEdGraphPin* OutPin = Pins[0];
+//		Outputs.Add(FNiagaraNodeResult(Compiler->GetExternalCurveConstant(ConstName), OutPin));
+//	}
 }
 
-void UNiagaraNodeConstant::PostEditChangeProperty(struct FPropertyChangedEvent& PropertyChangedEvent)
+void UDEPRECATED_NiagaraNodeConstant::PostEditChangeProperty(struct FPropertyChangedEvent& PropertyChangedEvent)
 {
 	static FName DataTypePropertyName(TEXT("DataType"));
 	if (PropertyChangedEvent.Property && PropertyChangedEvent.Property->GetFName() == DataTypePropertyName)
diff --git a/Engine/Source/Editor/UnrealEd/Private/Niagara/NiagaraNodeFunctionCall.cpp b/Engine/Source/Editor/UnrealEd/Private/Niagara/NiagaraNodeFunctionCall.cpp
new file mode 100644
--- /dev/null
+++ b/Engine/Source/Editor/UnrealEd/Private/Niagara/NiagaraNodeFunctionCall.cpp
@@ -0,0 +1,163 @@
+
+#include "UnrealEd.h"
+#include "BlueprintGraphDefinitions.h"
+#include "GraphEditorSettings.h"
+
+void UNiagaraNodeFunctionCall::PostLoad()
+{
+	Super::PostLoad();
+
+// 	if (FunctionScript)
+// 	{
+// 		FunctionScript->ConditionalPostLoad();
+// 	}
+}
+
+// void UNiagaraNodeFunctionCall::PostEditChangeProperty(struct FPropertyChangedEvent& PropertyChangedEvent)
+// {
+// 	ReallocatePins();
+// }
+
+void UNiagaraNodeFunctionCall::ReallocatePins()
+{
+// 	if (!FunctionScript)
+// 		return;
+// 
+// 	const UEdGraphSchema_Niagara* Schema = GetDefault<UEdGraphSchema_Niagara>();
+// 
+// // 	TArray<class UEdGraphPin*> PrevLinks;
+// // 	if (Pins.Num() > 0)
+// // 	{
+// // 		PrevLinks = Pins[0]->LinkedTo;
+// // 	}
+// 
+// 	Pins.Empty();
+// 	
+// 	if (UNiagaraScriptSource* Source = Cast<UNiagaraScriptSource>(FunctionScript->Source))
+// 	{
+// 		if (UNiagaraGraph* Graph = Cast<UNiagaraGraph>(Source->NodeGraph))
+// 		{
+// 			if (UNiagaraNodeOutput* OutNode = Graph->FindOutputNode())
+// 			{
+// 				//Add the input pins.
+// 				TArray<UNiagaraNodeInput*> InputNodes;
+// 				Graph->FindInputNodes(InputNodes);
+// 
+// 				for (UNiagaraNodeInput* InNode : InputNodes)
+// 				{
+// 					check(InNode);
+// 					switch (InNode->Input.Type)
+// 					{
+// 					case ENiagaraDataType::Scalar:
+// 					{
+// 						CreatePin(EGPD_Input, Schema->PC_Float, TEXT(""), NULL, false, false, InNode->Input.Name.ToString());
+// 					}
+// 						break;
+// 					case ENiagaraDataType::Vector:
+// 					{
+// 						CreatePin(EGPD_Input, Schema->PC_Vector, TEXT(""), NULL, false, false, InNode->Input.Name.ToString());
+// 					}
+// 						break;
+// 					case ENiagaraDataType::Matrix:
+// 					{
+// 						CreatePin(EGPD_Input, Schema->PC_Matrix, TEXT(""), NULL, false, false, InNode->Input.Name.ToString());
+// 					}
+// 						break;
+// 					};
+// 				}
+// 
+// 				//Add the output pins.
+// 				for (FNiagaraVariableInfo& Output : OutNode->Outputs)
+// 				{
+// 					switch (Output.Type)
+// 					{
+// 					case ENiagaraDataType::Scalar:
+// 					{
+// 						CreatePin(EGPD_Output, Schema->PC_Float, TEXT(""), NULL, false, false, Output.Name.ToString());
+// 					}
+// 						break;
+// 					case ENiagaraDataType::Vector:
+// 					{
+// 						CreatePin(EGPD_Output, Schema->PC_Vector, TEXT(""), NULL, false, false, Output.Name.ToString());
+// 					}
+// 						break;
+// 					case ENiagaraDataType::Matrix:
+// 					{
+// 						CreatePin(EGPD_Output, Schema->PC_Matrix, TEXT(""), NULL, false, false, Output.Name.ToString());
+// 					}
+// 						break;
+// 					};
+// 				}		
+// 			}
+// 		}
+// 	}
+
+// 	for (UEdGraphPin* Pin : Pins)
+// 	{
+// 		for (UEdGraphPin* PrevLinkedPin : PrevLinks)
+// 		{
+// 			if ((Pin->Direction == EGPD_Input && PrevLinkedPin->Direction == EGPD_Output)
+// 				|| (Pin->Direction == EGPD_Output && PrevLinkedPin->Direction == EGPD_Input))
+// 			{
+// 				if (Pin->PinName == PrevLinkedPin->PinName)
+// 				{
+// 				}
+// 			}			
+// 		}
+// 		if (Pins[0].PinName == Pin->PinName && Schema->CanCreateConnection(Pin, Pins[0]).Response == ECanCreateConnectionResponse::CONNECT_RESPONSE_MAKE)
+// 		{
+// 			Pins[0]->MakeLinkTo(Pin);
+// 		}
+// 	}
+}
+
+void UNiagaraNodeFunctionCall::AllocateDefaultPins()
+{
+	ReallocatePins();
+}
+
+// FText UNiagaraNodeFunctionCall::GetNodeTitle(ENodeTitleType::Type TitleType) const
+// {
+// 	return /*FunctionScript ? FText::FromString(FunctionScript->GetName()) : */Super::GetNodeTitle(TitleType);
+// }
+// 
+// FText UNiagaraNodeFunctionCall::GetTooltipText()const
+// {
+// 	const FNiagaraOpInfo* OpInfo = FNiagaraOpInfo::GetOpInfo(OpName);
+// 	check(OpInfo);
+// 	return OpInfo->Description;
+// }
+// 
+// FLinearColor UNiagaraNodeFunctionCall::GetNodeTitleColor() const
+// {
+// 	return GetDefault<UGraphEditorSettings>()->FunctionCallNodeTitleColor;
+// }
+
+void UNiagaraNodeFunctionCall::Compile(class INiagaraCompiler* Compiler, TArray<FNiagaraNodeResult>& Outputs)
+{
+// 	const FNiagaraOpInfo* OpInfo = FNiagaraOpInfo::GetOpInfo(OpName);
+// 	check(OpInfo);
+// 	int32 NumInputs = OpInfo->Inputs.Num();
+// 	int32 NumOutputs = OpInfo->Outputs.Num();
+// 
+// 	//No need to compile the input pins here. This will be done as a consequence of the function call compilation.
+// // 	TArray<TNiagaraExprPtr> Inputs;
+// // 	for (int32 i = 0; i < NumInputs; ++i)
+// // 	{
+// // 		UEdGraphPin *Pin = Pins[i];
+// // 		check(Pin->Direction == EGPD_Input);
+// // 		Inputs.Add(Compiler->CompilePin(Pin));
+// // 	}
+// 
+// 	TArray<TNiagaraExprPtr> OutputExpressions;
+// 	Compiler->FunctionCall(FunctionScript, OutputExpressions);
+// 
+// 	check(OutputExpressions.Num() == OpInfo->Outputs.Num());
+// 
+// 	for (int32 i = 0; i < NumOutputs; ++i)
+// 	{
+// 		UEdGraphPin *Pin = Pins[NumInputs + i];
+// 		check(Pin->Direction == EGPD_Output);
+// 		Outputs.Add(FNiagaraNodeResult(OutputExpressions[i], Pin));
+// 	}
+}
\ No newline at end of file
diff --git a/Engine/Source/Editor/UnrealEd/Private/Niagara/NiagaraNodeGetAttr.cpp b/Engine/Source/Editor/UnrealEd/Private/Niagara/NiagaraNodeGetAttr.cpp
--- a/Engine/Source/Editor/UnrealEd/Private/Niagara/NiagaraNodeGetAttr.cpp
+++ b/Engine/Source/Editor/UnrealEd/Private/Niagara/NiagaraNodeGetAttr.cpp
@@ -4,26 +4,26 @@
 #include "INiagaraCompiler.h"
 
 
-UNiagaraNodeGetAttr::UNiagaraNodeGetAttr(const FObjectInitializer& ObjectInitializer)
+UDEPRECATED_NiagaraNodeGetAttr::UDEPRECATED_NiagaraNodeGetAttr(const FObjectInitializer& ObjectInitializer)
 	: Super(ObjectInitializer)
 {
 }
 
-void UNiagaraNodeGetAttr::AllocateDefaultPins()
+void UDEPRECATED_NiagaraNodeGetAttr::AllocateDefaultPins()
 {
 	const UEdGraphSchema_Niagara* Schema = GetDefault<UEdGraphSchema_Niagara>();
 
 	CreatePin(EGPD_Output, Schema->PC_Vector, TEXT(""), NULL, false, false, AttrName.ToString());
 }
 
-FText UNiagaraNodeGetAttr::GetNodeTitle(ENodeTitleType::Type TitleType) const
+FText UDEPRECATED_NiagaraNodeGetAttr::GetNodeTitle(ENodeTitleType::Type TitleType) const
 {
 	FFormatNamedArguments Args;
 	Args.Add(TEXT("Attribute"), FText::FromName(AttrName));
 	return FText::Format(NSLOCTEXT("Niagara", "GetAttribute", "Get {Attribute}"), Args);
 }
 
-FLinearColor UNiagaraNodeGetAttr::GetNodeTitleColor() const
+FLinearColor UDEPRECATED_NiagaraNodeGetAttr::GetNodeTitleColor() const
 {
 	const UEdGraphSchema_Niagara* Schema = GetDefault<UEdGraphSchema_Niagara>();
 
@@ -37,8 +37,14 @@ FLinearColor UNiagaraNodeGetAttr::GetNodeTitleColor() const
 	}
 }
 
-void UNiagaraNodeGetAttr::Compile(class INiagaraCompiler* Compiler, TArray<FNiagaraNodeResult>& Outputs)
+void UDEPRECATED_NiagaraNodeGetAttr::Compile(class INiagaraCompiler* Compiler, TArray<FNiagaraNodeResult>& Outputs)
 {
-	Outputs.Add(FNiagaraNodeResult(Compiler->GetAttribute(AttrName), Pins[0]));
+	check(0);
+	//UE_LOG(LogNiagara, Fatal, TEXT("DEPRECATED_UNiagaraNodeGetAttr is deprecated. All instance of it should have been replaced in UNiagaraGraph::PostLoad()"));
+
+	return;
+
+	//TODO REMOVE THIS NODE ENTIRELY
+	//Outputs.Add(FNiagaraNodeResult(Compiler->GetAttribute(AttrName), Pins[0]));
 }
 
diff --git a/Engine/Source/Editor/UnrealEd/Private/Niagara/NiagaraNodeInput.cpp b/Engine/Source/Editor/UnrealEd/Private/Niagara/NiagaraNodeInput.cpp
new file mode 100644
--- /dev/null
+++ b/Engine/Source/Editor/UnrealEd/Private/Niagara/NiagaraNodeInput.cpp
@@ -0,0 +1,178 @@
+// Copyright 1998-2014 Epic Games, Inc. All Rights Reserved.
+
+#include "UnrealEd.h"
+#include "BlueprintGraphDefinitions.h"
+#include "GraphEditorSettings.h"
+
+UNiagaraNodeInput::UNiagaraNodeInput(const FObjectInitializer& ObjectInitializer)
+: Super(ObjectInitializer)
+, Input(NAME_None, ENiagaraDataType::Vector)
+, bExposeWhenConstant(true)
+, FloatDefault(1.0f)
+, VectorDefault(1.0f, 1.0f, 1.0f, 1.0f)
+, MatrixDefault(FMatrix::Identity)
+, bCanBeExposed(true)
+{
+}
+
+void UNiagaraNodeInput::PostEditChangeProperty(struct FPropertyChangedEvent& PropertyChangedEvent)
+{
+	ReallocatePins();
+}
+
+void UNiagaraNodeInput::ReallocatePins()
+{
+	const UEdGraphSchema_Niagara* Schema = GetDefault<UEdGraphSchema_Niagara>();
+	Modify();
+
+	// Break any links to 'orphan' pins
+	for (int32 PinIndex = 0; PinIndex < Pins.Num(); ++PinIndex)
+	{
+		UEdGraphPin* Pin = Pins[PinIndex];
+		TArray<class UEdGraphPin*>& LinkedToRef = Pin->LinkedTo;
+		for (int32 LinkIdx = 0; LinkIdx < LinkedToRef.Num(); LinkIdx++)
+		{
+			UEdGraphPin* OtherPin = LinkedToRef[LinkIdx];
+			// If we are linked to a pin that its owner doesn't know about, break that link
+			if (!OtherPin->GetOwningNode()->Pins.Contains(OtherPin))
+			{
+				Pin->LinkedTo.Remove(OtherPin);
+			}
+		}
+	}
+
+	// Store the old Input and Output pins
+	TArray<UEdGraphPin*> OldInputPins;
+	TArray<UEdGraphPin*> OldOutputPins;
+	GetInputPins(OldInputPins);
+	GetOutputPins(OldOutputPins);
+
+	// Move the existing pins to a saved array
+	TArray<UEdGraphPin*> OldPins(Pins);
+	Pins.Empty();
+
+
+	switch (Input.Type)
+	{
+	case ENiagaraDataType::Scalar:
+	{
+		CreatePin(EGPD_Output, Schema->PC_Float, TEXT(""), NULL, false, false, Input.Name.ToString());
+	}
+		break;
+	case ENiagaraDataType::Vector:
+	{
+		CreatePin(EGPD_Output, Schema->PC_Vector, TEXT(""), NULL, false, false, Input.Name.ToString());
+	}
+		break;
+	case ENiagaraDataType::Matrix:
+	{
+		CreatePin(EGPD_Output, Schema->PC_Matrix, TEXT(""), NULL, false, false, Input.Name.ToString());
+	}
+		break;
+	case ENiagaraDataType::Curve:
+	{
+		CreatePin(EGPD_Output, Schema->PC_Curve, TEXT(""), NULL, false, false, Input.Name.ToString());
+	}
+		break;
+	};
+
+	// Get new Input and Output pins
+	TArray<UEdGraphPin*> NewInputPins;
+	TArray<UEdGraphPin*> NewOutputPins;
+	GetInputPins(NewInputPins);
+	GetOutputPins(NewOutputPins);
+
+	for (int32 PinIndex = 0; PinIndex < OldInputPins.Num(); PinIndex++)
+	{
+		if (PinIndex < NewInputPins.Num())
+		{
+			NewInputPins[PinIndex]->CopyPersistentDataFromOldPin(*OldInputPins[PinIndex]);
+		}
+	}
+
+	for (int32 PinIndex = 0; PinIndex < OldOutputPins.Num(); PinIndex++)
+	{
+		if (PinIndex < NewOutputPins.Num())
+		{
+			NewOutputPins[PinIndex]->CopyPersistentDataFromOldPin(*OldOutputPins[PinIndex]);
+		}
+	}
+
+	OldInputPins.Empty();
+	OldOutputPins.Empty();
+
+	// Throw away the original pins
+	for (int32 OldPinIndex = 0; OldPinIndex < OldPins.Num(); ++OldPinIndex)
+	{
+		UEdGraphPin* OldPin = OldPins[OldPinIndex];
+		OldPin->Modify();
+		OldPin->BreakAllPinLinks();
+
+#if 0
+		UEdGraphNode::ReturnPinToPool(OldPin);
+#else
+		OldPin->Rename(NULL, GetTransientPackage(), REN_None);
+		OldPin->RemoveFromRoot();
+		OldPin->MarkPendingKill();
+#endif
+	}
+	OldPins.Empty();
+
+	GetGraph()->NotifyGraphChanged();
+}
+
+void UNiagaraNodeInput::AllocateDefaultPins()
+{
+	ReallocatePins();
+}
+
+FText UNiagaraNodeInput::GetNodeTitle(ENodeTitleType::Type TitleType) const
+{
+	return Input.Name == NAME_None ? NSLOCTEXT("NiagaraNodeInput", "Input", "Input") : FText::FromName(Input.Name);
+}
+
+FLinearColor UNiagaraNodeInput::GetNodeTitleColor() const
+{
+	if (IsConstant())
+	{
+		return CastChecked<UEdGraphSchema_Niagara>(GetSchema())->NodeTitleColor_Constant;
+	}
+	else
+	{
+		return CastChecked<UEdGraphSchema_Niagara>(GetSchema())->NodeTitleColor_Attribute;
+	}
+}
+
+void UNiagaraNodeInput::Compile(class INiagaraCompiler* Compiler, TArray<FNiagaraNodeResult>& Outputs)
+{
+	//Input nodes in scripts being used as functions by other scripts will actually never get this far.
+	//The compiler will find them and replace them with direct links between the pins going in and the nodes connected to the corresponding input.
+
+	if (IsConstant())
+	{
+		//treat this input as an external constant.
+		switch (Input.Type)
+		{
+		case ENiagaraDataType::Scalar:	Outputs.Add(FNiagaraNodeResult(Compiler->GetExternalConstant(Input, FloatDefault), Pins[0]));	break;
+		case ENiagaraDataType::Vector:	Outputs.Add(FNiagaraNodeResult(Compiler->GetExternalConstant(Input, VectorDefault), Pins[0]));	break;
+		case ENiagaraDataType::Matrix:	Outputs.Add(FNiagaraNodeResult(Compiler->GetExternalConstant(Input, MatrixDefault), Pins[0]));	break;
+		case ENiagaraDataType::Curve:	Outputs.Add(FNiagaraNodeResult(Compiler->GetExternalCurveConstant(Input), Pins[0]));	break;
+		}
+	}
+	else
+	{
+		//This input is an attribute.
+		Outputs.Add(FNiagaraNodeResult(Compiler->GetAttribute(Input), Pins[0]));
+	}
+}
+
+bool UNiagaraNodeInput::IsConstant()const
+{
+	//This input is a constant if it does not match an attribute for the current script.
+	return GetNiagaraGraph()->GetAttributeIndex(Input) == INDEX_NONE;
+}
+
+bool UNiagaraNodeInput::IsExposedConstant()const
+{
+	return bExposeWhenConstant && IsConstant();
+}
\ No newline at end of file
diff --git a/Engine/Source/Editor/UnrealEd/Private/Niagara/NiagaraNodeOutput.cpp b/Engine/Source/Editor/UnrealEd/Private/Niagara/NiagaraNodeOutput.cpp
new file mode 100644
--- /dev/null
+++ b/Engine/Source/Editor/UnrealEd/Private/Niagara/NiagaraNodeOutput.cpp
@@ -0,0 +1,161 @@
+// Copyright 1998-2014 Epic Games, Inc. All Rights Reserved.
+
+#include "UnrealEd.h"
+#include "BlueprintGraphDefinitions.h"
+#include "GraphEditorSettings.h"
+
+UNiagaraNodeOutput::UNiagaraNodeOutput(const FObjectInitializer& ObjectInitializer)
+: Super(ObjectInitializer)
+{
+}
+
+void UNiagaraNodeOutput::PostEditChangeProperty(struct FPropertyChangedEvent& PropertyChangedEvent)
+{
+	ReallocatePins();
+}
+
+void UNiagaraNodeOutput::ReallocatePins()
+{
+	Modify();
+
+	// Break any links to 'orphan' pins
+	for (int32 PinIndex = 0; PinIndex < Pins.Num(); ++PinIndex)
+	{
+		UEdGraphPin* Pin = Pins[PinIndex];
+		TArray<class UEdGraphPin*>& LinkedToRef = Pin->LinkedTo;
+		for (int32 LinkIdx = 0; LinkIdx < LinkedToRef.Num(); LinkIdx++)
+		{
+			UEdGraphPin* OtherPin = LinkedToRef[LinkIdx];
+			// If we are linked to a pin that its owner doesn't know about, break that link
+			if (!OtherPin->GetOwningNode()->Pins.Contains(OtherPin))
+			{
+				Pin->LinkedTo.Remove(OtherPin);
+			}
+		}
+	}
+
+	// Store the old Input and Output pins
+	TArray<UEdGraphPin*> OldInputPins;
+	TArray<UEdGraphPin*> OldOutputPins;
+	GetInputPins(OldInputPins);
+	GetOutputPins(OldOutputPins);
+
+	// Move the existing pins to a saved array
+	TArray<UEdGraphPin*> OldPins(Pins);
+	Pins.Empty();
+
+	// Recreate the new pins
+	AllocateDefaultPins();
+
+	// Get new Input and Output pins
+	TArray<UEdGraphPin*> NewInputPins;
+	TArray<UEdGraphPin*> NewOutputPins;
+	GetInputPins(NewInputPins);
+	GetOutputPins(NewOutputPins);
+
+	for (int32 PinIndex = 0; PinIndex < OldInputPins.Num(); PinIndex++)
+	{
+		if (PinIndex < NewInputPins.Num())
+		{
+			NewInputPins[PinIndex]->CopyPersistentDataFromOldPin(*OldInputPins[PinIndex]);
+		}
+	}
+
+	for (int32 PinIndex = 0; PinIndex < OldOutputPins.Num(); PinIndex++)
+	{
+		if (PinIndex < NewOutputPins.Num())
+		{
+			NewOutputPins[PinIndex]->CopyPersistentDataFromOldPin(*OldOutputPins[PinIndex]);
+		}
+	}
+
+	OldInputPins.Empty();
+	OldOutputPins.Empty();
+
+	// Throw away the original pins
+	for (int32 OldPinIndex = 0; OldPinIndex < OldPins.Num(); ++OldPinIndex)
+	{
+		UEdGraphPin* OldPin = OldPins[OldPinIndex];
+		OldPin->Modify();
+		OldPin->BreakAllPinLinks();
+
+#if 0
+		UEdGraphNode::ReturnPinToPool(OldPin);
+#else
+		OldPin->Rename(NULL, GetTransientPackage(), REN_None);
+		OldPin->RemoveFromRoot();
+		OldPin->MarkPendingKill();
+#endif
+	}
+	OldPins.Empty();
+
+	GetGraph()->NotifyGraphChanged();
+}
+
+void UNiagaraNodeOutput::AllocateDefaultPins()
+{
+	const UEdGraphSchema_Niagara* Schema = GetDefault<UEdGraphSchema_Niagara>();
+	for (const FNiagaraVariableInfo& Output : Outputs)
+	{
+		switch (Output.Type)
+		{
+		case ENiagaraDataType::Scalar:
+		{
+			CreatePin(EGPD_Input, Schema->PC_Float, TEXT(""), NULL, false, false, Output.Name.ToString());
+		}
+			break;
+		case ENiagaraDataType::Vector:
+		{
+			CreatePin(EGPD_Input, Schema->PC_Vector, TEXT(""), NULL, false, false, Output.Name.ToString());
+		}
+			break;
+		case ENiagaraDataType::Matrix:
+		{
+			CreatePin(EGPD_Input, Schema->PC_Matrix, TEXT(""), NULL, false, false, Output.Name.ToString());
+		}
+			break;
+		case ENiagaraDataType::Curve:
+		{
+			 CreatePin(EGPD_Input, Schema->PC_Curve, TEXT(""), NULL, false, false, Output.Name.ToString());
+		}
+			break;
+		};
+	}
+}
+
+FText UNiagaraNodeOutput::GetNodeTitle(ENodeTitleType::Type TitleType) const
+{
+	return NSLOCTEXT("NiagaraNode", "Output", "Output");
+}
+
+FLinearColor UNiagaraNodeOutput::GetNodeTitleColor() const
+{
+	return CastChecked<UEdGraphSchema_Niagara>(GetSchema())->NodeTitleColor_Attribute;
+}
+
+void UNiagaraNodeOutput::Compile(class INiagaraCompiler* Compiler, TArray<FNiagaraNodeResult>& OutputExpressions)
+{
+	TArray<TNiagaraExprPtr> InputExpressions;
+
+	for (int32 i = 0; i < Outputs.Num(); ++i)
+	{
+		const FNiagaraVariableInfo& Out = Outputs[i];
+		UEdGraphPin* Pin = Pins[i];
+		check(Pin->Direction == EGPD_Input);
+
+		TNiagaraExprPtr Result = Compiler->CompilePin(Pin);
+		if (!Result.IsValid())
+		{
+			//The pin was not connected so pass through the initial value of the attribute.
+			Result = Compiler->GetAttribute(Out);
+		}
+
+		InputExpressions.Add(Result);
+	}
+
+	check(Outputs.Num() == InputExpressions.Num());
+	for (int32 i = 0; i < Outputs.Num(); ++i)
+	{
+		OutputExpressions.Add(FNiagaraNodeResult(Compiler->Output(Outputs[i], InputExpressions[i]), Pins[i]));
+	}
+}
\ No newline at end of file
diff --git a/Engine/Source/Editor/UnrealEd/Private/Niagara/NiagaraNodeOutputUpdate.cpp b/Engine/Source/Editor/UnrealEd/Private/Niagara/NiagaraNodeOutputUpdate.cpp
--- a/Engine/Source/Editor/UnrealEd/Private/Niagara/NiagaraNodeOutputUpdate.cpp
+++ b/Engine/Source/Editor/UnrealEd/Private/Niagara/NiagaraNodeOutputUpdate.cpp
@@ -4,65 +4,72 @@
 #include "BlueprintGraphDefinitions.h"
 #include "GraphEditorSettings.h"
 
-UNiagaraNodeOutputUpdate::UNiagaraNodeOutputUpdate(const FObjectInitializer& ObjectInitializer)
+UDEPRECATED_NiagaraNodeOutputUpdate::UDEPRECATED_NiagaraNodeOutputUpdate(const FObjectInitializer& ObjectInitializer)
 	: Super(ObjectInitializer)
 {
 }
 
-void UNiagaraNodeOutputUpdate::AllocateDefaultPins()
+void UDEPRECATED_NiagaraNodeOutputUpdate::AllocateDefaultPins()
 {
-	const UEdGraphSchema_Niagara* Schema = GetDefault<UEdGraphSchema_Niagara>();
-
-	UNiagaraScriptSource* Source = GetSource();
-
-	TArray<FName> OutputNames;
-	Source->GetParticleAttributes(OutputNames);
-
-	for (int32 i = 0; i < OutputNames.Num(); i++)
-	{
-		UEdGraphPin* Pin = CreatePin(EGPD_Input, Schema->PC_Vector, TEXT(""), NULL, false, false, OutputNames[i].ToString());
-		Pin->bDefaultValueIsIgnored = true;
-	}
+// 	const UEdGraphSchema_Niagara* Schema = GetDefault<UEdGraphSchema_Niagara>();
+// 
+// 	UNiagaraScriptSource* Source = GetSource();
+// 
+// 	TArray<FName> OutputNames;
+// 	Source->GetParticleAttributes(OutputNames);
+// 
+// 	for (int32 i = 0; i < OutputNames.Num(); i++)
+// 	{
+// 		UEdGraphPin* Pin = CreatePin(EGPD_Input, Schema->PC_Vector, TEXT(""), NULL, false, false, OutputNames[i].ToString());
+// 		Pin->bDefaultValueIsIgnored = true;
+// 	}
 }
 
-FText UNiagaraNodeOutputUpdate::GetNodeTitle(ENodeTitleType::Type TitleType) const
+FText UDEPRECATED_NiagaraNodeOutputUpdate::GetNodeTitle(ENodeTitleType::Type TitleType) const
 {
 	return NSLOCTEXT("NiagaraNode", "Output", "Output");
 }
 
-FLinearColor UNiagaraNodeOutputUpdate::GetNodeTitleColor() const
+FLinearColor UDEPRECATED_NiagaraNodeOutputUpdate::GetNodeTitleColor() const
 {
 	return GetDefault<UGraphEditorSettings>()->FunctionTerminatorNodeTitleColor;
 }
 
-void UNiagaraNodeOutputUpdate::Compile(class INiagaraCompiler* Compiler, TArray<FNiagaraNodeResult>& Outputs)
+void UDEPRECATED_NiagaraNodeOutputUpdate::Compile(class INiagaraCompiler* Compiler, TArray<FNiagaraNodeResult>& Outputs)
 {
-	TArray<FName> OutputNames;
-	Compiler->GetParticleAttributes(OutputNames);
-	int32 NumOutputs = OutputNames.Num();
-
-	//No output pins so we have 1 input pin per attribute.
-	check(OutputNames.Num() == Pins.Num());
-
-	TArray<TNiagaraExprPtr> InputExpressions;
-	for (int32 i = 0; i < NumOutputs; ++i)
-	{
-		UEdGraphPin* Pin = Pins[i];
-		check( Pin->Direction == EGPD_Input );
+	//UE_LOG(LogNiagara, Fatal, TEXT("DEPRECATED_UNiagaraNodeOutputUpdate is deprecated. All instance of it should have been replaced in UNiagaraGraph::PostLoad()"));
+	check(0);
 
-		TNiagaraExprPtr Output = Compiler->CompilePin(Pin);
-		if (!Output.IsValid())
-		{
-			//The pin was not connected so pass through the initial value of the attribute.
-			Output = Compiler->GetAttribute(OutputNames[i]);
-		}
+	return;
 
-		InputExpressions.Add(Output);
-	}
+	//TODO REMOVE THIS NODE ENTIRELY
 
-	check(NumOutputs == InputExpressions.Num());
-	for (int32 i = 0; i < NumOutputs; ++i)
-	{
-		Compiler->Output(OutputNames[i], InputExpressions[i]);		
-	}
+// 	TArray<FName> OutputNames;
+// 	Compiler->GetParticleAttributes(OutputNames);
+// 	int32 NumOutputs = OutputNames.Num();
+// 
+// 	//No output pins so we have 1 input pin per attribute.
+// 	check(OutputNames.Num() == Pins.Num());
+// 
+// 	TArray<TNiagaraExprPtr> InputExpressions;
+// 	for (int32 i = 0; i < NumOutputs; ++i)
+// 	{
+// 		UEdGraphPin* Pin = Pins[i];
+// 		check( Pin->Direction == EGPD_Input );
+// 
+// 		TNiagaraExprPtr Output = Compiler->CompilePin(Pin);
+// 		if (!Output.IsValid())
+// 		{
+// 			//The pin was not connected so pass through the initial value of the attribute.
+// 			Output = Compiler->GetAttribute(OutputNames[i]);
+// 		}
+// 
+// 		InputExpressions.Add(Output);
+// 	}
+// 
+// 	check(NumOutputs == InputExpressions.Num());
+// 	for (int32 i = 0; i < NumOutputs; ++i)
+// 	{
+// 		Compiler->Output(OutputNames[i], InputExpressions[i]);		
+// 	}
 }
\ No newline at end of file
diff --git a/Engine/Source/Editor/UnrealEd/Private/Niagara/NiagaraScriptSource.cpp b/Engine/Source/Editor/UnrealEd/Private/Niagara/NiagaraScriptSource.cpp
--- a/Engine/Source/Editor/UnrealEd/Private/Niagara/NiagaraScriptSource.cpp
+++ b/Engine/Source/Editor/UnrealEd/Private/Niagara/NiagaraScriptSource.cpp
@@ -20,6 +20,55 @@ class UNiagaraScriptSource* UNiagaraGraph::GetSource() const
 	return CastChecked<UNiagaraScriptSource>(GetOuter());
 }
 
+UNiagaraNodeOutput* UNiagaraGraph::FindOutputNode() const
+{
+	for (UEdGraphNode* Node : Nodes)
+	{
+		if (UNiagaraNodeOutput* OutNode = Cast<UNiagaraNodeOutput>(Node))
+		{
+			return OutNode;
+		}
+	}
+	check(0);
+	return NULL;
+}
+
+void UNiagaraGraph::FindInputNodes(TArray<class UNiagaraNodeInput*>& OutInputNodes) const
+{
+	for (UEdGraphNode* Node : Nodes)
+	{
+		if (UNiagaraNodeInput* InNode = Cast<UNiagaraNodeInput>(Node))
+		{
+			OutInputNodes.Add(InNode);
+		}
+	}
+}
+
+int32 UNiagaraGraph::GetAttributeIndex(const FNiagaraVariableInfo& Attr)const
+{
+	const UNiagaraNodeOutput* OutNode = FindOutputNode();
+	check(OutNode);
+	for (int32 i = 0; i < OutNode->Outputs.Num(); ++i)
+	{
+		if (OutNode->Outputs[i] == Attr)
+		{
+			return i;
+		}
+	}
+	return INDEX_NONE;
+}
+
+void UNiagaraGraph::GetAttributes(TArray< FNiagaraVariableInfo >& OutAttributes)const
+{
+	const UNiagaraNodeOutput* OutNode = FindOutputNode();
+	check(OutNode);
+
+	for (const FNiagaraVariableInfo& Attr : OutNode->Outputs)
+	{
+		check(!OutAttributes.Find(Attr));
+		OutAttributes.Add(Attr);
+	}
+}
 
 //////////////////////////////////////////////////////////////////////////
 // UNiagraScriptSource
@@ -32,15 +81,17 @@ UNiagaraScriptSource::UNiagaraScriptSource(const FObjectInitializer& ObjectIniti
 void UNiagaraScriptSource::PostLoad()
 {
 	Super::PostLoad();
+
+#if WITH_EDITOR
 	UNiagaraScript* ScriptOwner = Cast<UNiagaraScript>(GetOuter());
-	if (ScriptOwner)
+	if (ScriptOwner && ScriptOwner->ByteCode.Num() == 0)
 	{
 		ScriptOwner->ConditionalPostLoad();
 		//FNiagaraEditorModule& NiagaraEditorModule = FModuleManager::Get().LoadModuleChecked<FNiagaraEditorModule>(TEXT("NiagaraEditor"));
-		//NiagaraEditorModule.CompileScript(ScriptOwner);
+		//NiagaraEditorModule.CompileScript(ScriptOwner);	
 		Compile();
 	}
-
+#endif
 }
 
 
@@ -53,40 +104,28 @@ void UNiagaraScriptSource::Compile()
 	ExposedVectorConstants.Empty();
 
 	// grab all constant nodes that are exposed to the editor
-	TArray<UNiagaraNodeConstant*> ConstNodes;
-	UpdateGraph->GetNodesOfClass<UNiagaraNodeConstant>(ConstNodes);
-	for (UNiagaraNodeConstant *Node : ConstNodes)
+	TArray<UNiagaraNodeInput*> ConstNodes;
+	NodeGraph->GetNodesOfClass<UNiagaraNodeInput>(ConstNodes);
+	for (UNiagaraNodeInput *Node : ConstNodes)
 	{
-		if (Node->bExposeToEffectEditor)
+		if (Node->IsExposedConstant())
 		{
-			if(Node->DataType == ENiagaraDataType::Vector)
+			if(Node->Input.Type == ENiagaraDataType::Vector)
 			{
 				EditorExposedVectorConstant *Const = new EditorExposedVectorConstant();
-				Const->ConstName = Node->ConstName;
+				Const->ConstName = Node->Input.Name;
 				ExposedVectorConstants.Add( MakeShareable(Const) );
 			}
-			else if(Node->DataType == ENiagaraDataType::Curve)
+			else if (Node->Input.Type == ENiagaraDataType::Curve)
 			{
 				EditorExposedVectorCurveConstant *Const = new EditorExposedVectorCurveConstant();
-				Const->ConstName = Node->ConstName;
+				Const->ConstName = Node->Input.Name;
 				ExposedVectorCurveConstants.Add(MakeShareable(Const));
 			}
 		}
 	}
 }
 
-
-void UNiagaraScriptSource::GetParticleAttributes(TArray<FName>& VectorOutputs)
-{
-	VectorOutputs.Empty();
-
-	VectorOutputs.Add(FName(TEXT("Particle Position")));
-	VectorOutputs.Add(FName(TEXT("Particle Velocity")));
-	VectorOutputs.Add(FName(TEXT("Particle Color")));
-	VectorOutputs.Add(FName(TEXT("Particle Rotation")));
-	VectorOutputs.Add(FName(TEXT("Particle Age")));
-}
-
 void UNiagaraScriptSource::GetEmitterAttributes(TArray<FName>& VectorInputs, TArray<FName>& MatrixInputs)
 {
 	for (uint32 i=0; i < NiagaraConstants::NumBuiltinConstants; i++)
diff --git a/Engine/Source/Editor/UnrealEd/Public/INiagaraCompiler.h b/Engine/Source/Editor/UnrealEd/Public/INiagaraCompiler.h
--- a/Engine/Source/Editor/UnrealEd/Public/INiagaraCompiler.h
+++ b/Engine/Source/Editor/UnrealEd/Public/INiagaraCompiler.h
@@ -14,23 +14,20 @@ public:
 
 	/** Traverses the pins input, recursively compiling them and then adds any expressions needed for the owner node.*/
 	virtual TNiagaraExprPtr CompilePin(UEdGraphPin* Pin) = 0;
-
-	/** Fills OutAttributes with the names of all attributes for this emitter. */
-	virtual void GetParticleAttributes(TArray<FName>& OutAttributes) = 0;
 	
 	/** Creates an expression which fetches the named attribute. */
-	virtual TNiagaraExprPtr GetAttribute(FName AttributeName) = 0;
+	virtual TNiagaraExprPtr GetAttribute(const FNiagaraVariableInfo& Attribute) = 0;
 	/** Creates an expression which fetches the named constant. */
-	virtual TNiagaraExprPtr GetExternalConstant(FName ConstantName, float Default) = 0;
+	virtual TNiagaraExprPtr GetExternalConstant(const FNiagaraVariableInfo& Constant, float Default) = 0;
 	/** Creates an expression which fetches the named constant. */
-	virtual TNiagaraExprPtr GetExternalConstant(FName ConstantName, FVector4 Default) = 0;
+	virtual TNiagaraExprPtr GetExternalConstant(const FNiagaraVariableInfo& Constant, const FVector4& Default) = 0;
 	/** Creates an expression which fetches the named constant. */
-	virtual TNiagaraExprPtr GetExternalConstant(FName ConstantName, const FMatrix& Default) = 0;
+	virtual TNiagaraExprPtr GetExternalConstant(const FNiagaraVariableInfo& Constant, const FMatrix& Default) = 0;
 	/** Creates an expression which fetches the named CURVE. */
-	virtual TNiagaraExprPtr GetExternalCurveConstant(FName ConstantName) = 0;
+	virtual TNiagaraExprPtr GetExternalCurveConstant(const FNiagaraVariableInfo& Constant) = 0;
 
 	/** Creates an expression which writes the result of SourceExpression to the named attribute. */
-	virtual TNiagaraExprPtr Output(FName OutputName, TNiagaraExprPtr& SourceExpression) = 0;
+	virtual TNiagaraExprPtr Output(const FNiagaraVariableInfo& Attr, TNiagaraExprPtr& SourceExpression) = 0;
 
 	/** Checks the validity of the input expressions against the informations about this operation in FNiagaraOpInfo. */
 	virtual void CheckInputs(FName OpName, TArray<TNiagaraExprPtr>& Inputs) = 0;
diff --git a/Engine/Source/Editor/UnrealEd/Public/NiagaraEditorCommon.h b/Engine/Source/Editor/UnrealEd/Public/NiagaraEditorCommon.h
--- a/Engine/Source/Editor/UnrealEd/Public/NiagaraEditorCommon.h
+++ b/Engine/Source/Editor/UnrealEd/Public/NiagaraEditorCommon.h
@@ -80,12 +80,12 @@ typedef TSharedPtr<class FNiagaraExpression> TNiagaraExprPtr;
 class FNiagaraExpression : public TSharedFromThis<FNiagaraExpression>
 {
 public:
-	FNiagaraExpression(class FNiagaraCompiler* InCompiler, ENiagaraDataType InResultType)
+	FNiagaraExpression(class FNiagaraCompiler* InCompiler, const FNiagaraVariableInfo& InResult)
 		: Compiler(InCompiler)
 		, ResultLocation(ENiagaraExpressionResultLocation::Unknown)
 		, ResultIndex(INDEX_NONE)
 		, ComponentIndex(INDEX_NONE)
-		, ResultType(InResultType)
+		, Result(InResult)
 	{
 	}
 
@@ -115,8 +115,8 @@ public:
 	/** The index of the component we're accessing. Typically this is unused and set to INDEX_NONE. Scalar ops will set and use it however. */
 	int32 ComponentIndex;
 
-	/** The data type of the result. */
-	ENiagaraDataType ResultType;
+	/** The variable info for the result of this expression. */
+	FNiagaraVariableInfo Result;
 
 	//Expressions that this one relies on. Will be processed before this one.
 	TArray<TNiagaraExprPtr> SourceExpressions;
@@ -155,6 +155,7 @@ public:
 		case ENiagaraDataType::Scalar: Default = TEXT("1.0f");	break;
 		case ENiagaraDataType::Vector: Default = TEXT("1.0f,1.0f,1.0f,1.0f");	break;
 		case ENiagaraDataType::Matrix: Default = TEXT("1.0f,0.0f,0.0f,0.0f,0.0f,1.0f,0.0f,0.0f,0.0f,0.0f,1.0f,0.0f,0.0f,0.0f,0.0f,1.0f");	break;
+		case ENiagaraDataType::Curve: Default = TEXT("");	break;
 		};
 	}
 };
diff --git a/Engine/Source/Runtime/Engine/Classes/Components/NiagaraComponent.h b/Engine/Source/Runtime/Engine/Classes/Components/NiagaraComponent.h
--- a/Engine/Source/Runtime/Engine/Classes/Components/NiagaraComponent.h
+++ b/Engine/Source/Runtime/Engine/Classes/Components/NiagaraComponent.h
@@ -1,14 +1,10 @@
 // Copyright 1998-2015 Epic Games, Inc. All Rights Reserved.
 
 #pragma once
-#include "NiagaraComponent.generated.h"
-
-/* Predefined particle attribute names */
-const FName ParticleAttr_Position = "Position";
-const FName ParticleAttr_Velocity = "Velocity";
-const FName ParticleAttr_Color = "Color";
 
+#include "Niagara/NiagaraCommon.h"
 
+#include "NiagaraComponent.generated.h"
 
 /**
 * Contains the simulation data for a single FNiagaraSimulation
@@ -19,13 +15,9 @@ class FNiagaraEmitterParticleData
 public:
 	FNiagaraEmitterParticleData()
 		: CurrentBuffer(0),
-		NumParticles(0)
+		NumParticles(0),
+		ParticleAllocation(0)
 	{
-		AddAttribute("Position");
-		AddAttribute("Velocity");
-		AddAttribute("Color");
-		AddAttribute("Rotation");
-		AddAttribute("Age");
 	}
 
 	~FNiagaraEmitterParticleData() {}
@@ -35,30 +27,43 @@ public:
 		ParticleBuffers[CurrentBuffer].Reset(NumExpectedParticles * AttrMap.Num());
 		ParticleBuffers[CurrentBuffer].AddUninitialized(NumExpectedParticles * AttrMap.Num());
 
-		//Can we make this a ring buffer and cut memory use in half?
-
 		ParticleAllocation = NumExpectedParticles;
 	}
 
-	const FVector4 *GetAttributeData(FName Name) const
+	void Reset()
 	{
-		int32 Offset = AttrMap[Name] * ParticleAllocation;
-		return ParticleBuffers[CurrentBuffer].GetData() + Offset;
+		AttrMap.Empty();
+		NumParticles = 0;
+		ParticleAllocation = 0;
+		CurrentBuffer = 0;
 	}
 
-	FVector4 *GetAttributeDataWrite(FName Name) 
+	const FVector4 *GetAttributeData(const FNiagaraVariableInfo& AttrID) const
 	{
-		int32 Offset = AttrMap[Name] * ParticleAllocation;
-		return ParticleBuffers[CurrentBuffer].GetData() + Offset;
+		const uint32* Offset = AttrMap.Find(AttrID);
+		return Offset ? ParticleBuffers[CurrentBuffer].GetData() + (*Offset * ParticleAllocation) : NULL;
 	}
 
-	
+	FVector4 *GetAttributeDataWrite(const FNiagaraVariableInfo& AttrID)
+	{
+		const uint32* Offset = AttrMap.Find(AttrID);
+		return Offset ? ParticleBuffers[CurrentBuffer].GetData() + (*Offset * ParticleAllocation) : NULL;
+	}
 
+	bool HasAttriubte(const FNiagaraVariableInfo& AttrID)
+	{
+		return AttrMap.Find(AttrID) != NULL;
+	}
 
-	void AddAttribute(FName NewAttrName)
+	void SetAttributes(const TArray<FNiagaraVariableInfo>& Attrs)
 	{
-		uint32 Idx = AttrMap.Num();
-		AttrMap.Add(NewAttrName, Idx);
+		Reset();
+		for (const FNiagaraVariableInfo& Attr : Attrs)
+		{
+			//TODO: Support attributes of any type. This may as well wait until the move of matrix ops etc over into UNiagaraScriptStructs.
+			uint32 Idx = AttrMap.Num();
+			AttrMap.Add(Attr, Idx);
+		}
 	}
 
 	int GetNumAttributes()				{ return AttrMap.Num(); }
@@ -71,12 +76,11 @@ public:
 	FVector4 *GetCurrentBuffer()		{ return ParticleBuffers[CurrentBuffer].GetData(); }
 	FVector4 *GetPreviousBuffer()		{ return ParticleBuffers[CurrentBuffer^0x1].GetData(); }
 
-	int GetBytesUsed()	{ return (ParticleBuffers[0].Num() + ParticleBuffers[1].Num()) * 16 + Attributes.Num() * 4; }
+	int GetBytesUsed()	{ return (ParticleBuffers[0].Num() + ParticleBuffers[1].Num()) * 16 + AttrMap.Num() * 4; }
 private:
 	uint32 CurrentBuffer, NumParticles, ParticleAllocation;
 	TArray<FVector4> ParticleBuffers[2];
-	TArray<FVector4*> Attributes;
-	TMap<FName, uint32> AttrMap;
+	TMap<FNiagaraVariableInfo, uint32> AttrMap;
 };
 
 
@@ -122,6 +126,9 @@ public:
 	virtual void PostEditChangeProperty(FPropertyChangedEvent& PropertyChangedEvent) override;
 #endif // WITH_EDITOR
 	// End UObject interface.
+
+	static const TArray<FNiagaraVariableInfo>& GetSystemConstants();
+
 };
 
 
diff --git a/Engine/Source/Runtime/Engine/Classes/Engine/NiagaraEffect.h b/Engine/Source/Runtime/Engine/Classes/Engine/NiagaraEffect.h
--- a/Engine/Source/Runtime/Engine/Classes/Engine/NiagaraEffect.h
+++ b/Engine/Source/Runtime/Engine/Classes/Engine/NiagaraEffect.h
@@ -108,19 +108,19 @@ public:
 
 	ENGINE_API TSharedPtr<FNiagaraSimulation> AddEmitter(FNiagaraEmitterProperties *Properties);
 
-	void SetConstant(FName ConstantName, const float Value)
+	void SetConstant(FNiagaraVariableInfo ID, const float Value)
 	{
-		Constants.SetOrAdd(ConstantName, Value);
+		Constants.SetOrAdd(ID, Value);
 	}
 
-	void SetConstant(FName ConstantName, const FVector4& Value)
+	void SetConstant(FNiagaraVariableInfo ID, const FVector4& Value)
 	{
-		Constants.SetOrAdd(ConstantName, Value);
+		Constants.SetOrAdd(ID, Value);
 	}
 
-	void SetConstant(FName ConstantName, const FMatrix& Value)
+	void SetConstant(FNiagaraVariableInfo ID, const FMatrix& Value)
 	{
-		Constants.SetOrAdd(ConstantName, Value);
+		Constants.SetOrAdd(ID, Value);
 	}
 
 	void SetConstant(FName ConstantName, FNiagaraDataObject *Value)
diff --git a/Engine/Source/Runtime/Engine/Classes/Engine/NiagaraEffectRenderer.h b/Engine/Source/Runtime/Engine/Classes/Engine/NiagaraEffectRenderer.h
--- a/Engine/Source/Runtime/Engine/Classes/Engine/NiagaraEffectRenderer.h
+++ b/Engine/Source/Runtime/Engine/Classes/Engine/NiagaraEffectRenderer.h
@@ -62,6 +62,10 @@ public:
 
 	virtual bool HasDynamicData() = 0;
 
+	virtual bool SetMaterialUsage() = 0;
+
+	virtual const TArray<FNiagaraVariableInfo>& GetRequiredAttributes() = 0;
+
 	virtual FPrimitiveViewRelevance GetViewRelevance(const FSceneView* View, FNiagaraSceneProxy *SceneProxy)
 	{
 		FPrimitiveViewRelevance Result;
@@ -84,7 +88,7 @@ public:
 	void SetMaterial(UMaterial *InMaterial, ERHIFeatureLevel::Type FeatureLevel)
 	{
 		Material = InMaterial;
-		if (!Material)
+		if (!Material || !SetMaterialUsage())
 		{
 			Material = UMaterial::GetDefaultMaterial(MD_Surface);
 		}
@@ -135,7 +139,7 @@ public:
 	virtual void CreateRenderThreadResources() override;
 
 	virtual void GetDynamicMeshElements(const TArray<const FSceneView*>& Views, const FSceneViewFamily& ViewFamily, uint32 VisibilityMap, FMeshElementCollector& Collector, const FNiagaraSceneProxy *SceneProxy) const override;
-
+	virtual bool SetMaterialUsage();
 	/** Update render data buffer from attributes */
 	FNiagaraDynamicDataBase *GenerateVertexData(const FNiagaraEmitterParticleData &Data) override;
 
@@ -144,7 +148,9 @@ public:
 	bool HasDynamicData();
 
 	UClass *GetPropertiesClass() override { return UNiagaraSpriteRendererProperties::StaticClass(); }
-	
+
+	virtual const TArray<FNiagaraVariableInfo>& GetRequiredAttributes();
+
 private:
 	UNiagaraSpriteRendererProperties *Properties;
 	struct FNiagaraDynamicDataSprites *DynamicDataRender;
@@ -175,6 +181,8 @@ public:
 
 	virtual void GetDynamicMeshElements(const TArray<const FSceneView*>& Views, const FSceneViewFamily& ViewFamily, uint32 VisibilityMap, FMeshElementCollector& Collector, const FNiagaraSceneProxy *SceneProxy) const override;
 
+	virtual bool SetMaterialUsage();
+
 	/** Update render data buffer from attributes */
 	FNiagaraDynamicDataBase *GenerateVertexData(const FNiagaraEmitterParticleData &Data);
 
@@ -202,6 +210,8 @@ public:
 	int GetDynamicDataSize();
 	bool HasDynamicData();
 
+	virtual const TArray<FNiagaraVariableInfo>& GetRequiredAttributes();
+
 	UClass *GetPropertiesClass() override { return UNiagaraRibbonRendererProperties::StaticClass(); }
 
 private:
diff --git a/Engine/Source/Runtime/Engine/Classes/Engine/NiagaraScript.h b/Engine/Source/Runtime/Engine/Classes/Engine/NiagaraScript.h
--- a/Engine/Source/Runtime/Engine/Classes/Engine/NiagaraScript.h
+++ b/Engine/Source/Runtime/Engine/Classes/Engine/NiagaraScript.h
@@ -21,7 +21,8 @@ class UNiagaraScript : public UObject
 	FNiagaraScriptConstantData ConstantData;
 
 	/** Attributes used by this script. */
-	TArray<FName> Attributes;
+	UPROPERTY()
+ 	TArray<FNiagaraVariableInfo> Attributes;
 
 #if WITH_EDITORONLY_DATA
 	/** 'Source' data/graphs for this script */
diff --git a/Engine/Source/Runtime/Engine/Classes/Engine/NiagaraSimulation.h b/Engine/Source/Runtime/Engine/Classes/Engine/NiagaraSimulation.h
--- a/Engine/Source/Runtime/Engine/Classes/Engine/NiagaraSimulation.h
+++ b/Engine/Source/Runtime/Engine/Classes/Engine/NiagaraSimulation.h
@@ -83,6 +83,8 @@ public:
 	virtual ~FNiagaraSimulation()
 	{}
 
+	void Init();
+
 	void Tick(float DeltaSeconds);
 
 
@@ -96,19 +98,21 @@ public:
 
 	NiagaraEffectRenderer *GetEffectRenderer()	{ return EffectRenderer; }
 	
-	bool IsEnabled()	{ return bIsEnabled;  }
-	void SetEnabled(bool bInEnabled)	{ bIsEnabled = bInEnabled;  }
+	bool IsEnabled()const 	{ return bIsEnabled;  }
 
 	ENGINE_API void SetRenderModuleType(EEmitterRenderModuleType Type, ERHIFeatureLevel::Type FeatureLevel);
 
 	int32 GetNumParticles()	{ return Data.GetNumParticles(); }
 
 	FNiagaraEmitterProperties *GetProperties()	{ return Props; }
-	void SetProperties(FNiagaraEmitterProperties *InProps)	{ Props = InProps; }
+	void SetProperties(FNiagaraEmitterProperties *InProps);
 
 	ENGINE_API float GetTotalCPUTime();
 	ENGINE_API int	GetTotalBytesUsed();
 
+	ENGINE_API void SetSpawnScript(UNiagaraScript* Script);
+	ENGINE_API void SetUpdateScript(UNiagaraScript* Script);
+	ENGINE_API void SetEnabled(bool bEnabled);
 private:
 	FNiagaraEmitterProperties *Props;		// points to an entry in the array of FNiagaraProperties stored in the EffectInstance (itself pointing to the effect's properties)
 
@@ -153,10 +157,6 @@ private:
 
 	int32 SpawnAndKillParticles(int32 NumToSpawn);
 
-
-	/** Spawn a new particle at this index */
-	int32 SpawnParticles(int32 NumToSpawn);
-
 	/** Util to move a particle */
 	void MoveParticleToIndex(int32 SrcIndex, int32 DestIndex)
 	{
@@ -171,4 +171,5 @@ private:
 		}
 	}
 
+	bool CheckAttriubtesForRenderer();
 };
diff --git a/Engine/Source/Runtime/Engine/Classes/Niagara/NiagaraCommon.h b/Engine/Source/Runtime/Engine/Classes/Niagara/NiagaraCommon.h
--- a/Engine/Source/Runtime/Engine/Classes/Niagara/NiagaraCommon.h
+++ b/Engine/Source/Runtime/Engine/Classes/Niagara/NiagaraCommon.h
@@ -4,6 +4,8 @@
 
 #include "NiagaraCommon.generated.h"
 
+DECLARE_LOG_CATEGORY_EXTERN(LogNiagara, Log, Verbose);
+
 UENUM()
 enum class ENiagaraDataType : uint8
 {
@@ -11,6 +13,47 @@ enum class ENiagaraDataType : uint8
 	Vector,
 	Matrix,
 	Curve,
-	Unknown,
-	Invalid,
 };
+
+USTRUCT()
+struct FNiagaraVariableInfo
+{
+	GENERATED_USTRUCT_BODY()
+
+	FNiagaraVariableInfo()
+	: Name(NAME_None)
+	, Type(ENiagaraDataType::Vector)
+	{}
+
+	FNiagaraVariableInfo(FName InName, ENiagaraDataType InType)
+	: Name(InName)
+	, Type(InType)
+	{}
+
+	UPROPERTY(EditAnywhere, Category = "Variable")
+	FName Name;
+
+	UPROPERTY(EditAnywhere, Category = "Variable")
+	TEnumAsByte<ENiagaraDataType> Type;
+	
+	FORCEINLINE bool operator==(const FNiagaraVariableInfo& Other)const
+	{
+		return Name == Other.Name && Type == Other.Type;
+	}
+
+	FORCEINLINE bool operator!=(const FNiagaraVariableInfo& Other)const
+	{
+		return !(*this == Other);
+	}
+};
+
+FORCEINLINE FArchive& operator<<(FArchive& Ar, FNiagaraVariableInfo& VarInfo)
+{
+	Ar << VarInfo.Name << VarInfo.Type;
+	return Ar;
+}
+
+FORCEINLINE uint32 GetTypeHash(const FNiagaraVariableInfo& Var)
+{
+	return HashCombine(GetTypeHash(Var.Name), (uint32)(Var.Type.GetValue()));
+}
diff --git a/Engine/Source/Runtime/Engine/Classes/Niagara/NiagaraConstantSet.h b/Engine/Source/Runtime/Engine/Classes/Niagara/NiagaraConstantSet.h
--- a/Engine/Source/Runtime/Engine/Classes/Niagara/NiagaraConstantSet.h
+++ b/Engine/Source/Runtime/Engine/Classes/Niagara/NiagaraConstantSet.h
@@ -12,43 +12,74 @@ struct FNiagaraConstantMap
 	GENERATED_USTRUCT_BODY()
 
 private:
-	TMap<FName, float> ScalarConstants;
-	TMap<FName, FVector4> VectorConstants;
-	TMap<FName, FMatrix> MatrixConstants;
-	TMap<FName, class FNiagaraDataObject*> DataConstants;
+	TMap<FNiagaraVariableInfo, float> ScalarConstants;
+	TMap<FNiagaraVariableInfo, FVector4> VectorConstants;
+	TMap<FNiagaraVariableInfo, FMatrix> MatrixConstants;
+	TMap<FNiagaraVariableInfo, class FNiagaraDataObject*> DataConstants;
 
 public:
 
+	void Empty()
+	{
+		ScalarConstants.Empty();
+		VectorConstants.Empty();
+		MatrixConstants.Empty();
+	}
+
+	void SetOrAdd(FNiagaraVariableInfo ID, float Sc)
+	{
+		check(ID.Type == ENiagaraDataType::Scalar);
+		ScalarConstants.FindOrAdd(ID) = Sc;
+	}
+
+	void SetOrAdd(FNiagaraVariableInfo ID, const FVector4& Vc)
+	{
+		check(ID.Type == ENiagaraDataType::Vector);
+		VectorConstants.FindOrAdd(ID) = Vc;
+	}
+
+	void SetOrAdd(FNiagaraVariableInfo ID, const FMatrix& Mc)
+	{
+		check(ID.Type == ENiagaraDataType::Matrix);
+		MatrixConstants.FindOrAdd(ID) = Mc;
+	}
+
+	void SetOrAdd(FNiagaraVariableInfo ID, FNiagaraDataObject* Cc)
+	{
+		check(ID.Type == ENiagaraDataType::Curve);
+		DataConstants.FindOrAdd(ID) = Cc;
+	}
+
 	void SetOrAdd(FName Name, float Sc)
 	{
-		ScalarConstants.FindOrAdd(Name) = Sc;
+		ScalarConstants.FindOrAdd(FNiagaraVariableInfo(Name, ENiagaraDataType::Scalar)) = Sc;
 	}
 
 	void SetOrAdd(FName Name, const FVector4& Vc)
 	{
-		VectorConstants.FindOrAdd(Name) = Vc;
+		VectorConstants.FindOrAdd(FNiagaraVariableInfo(Name, ENiagaraDataType::Vector)) = Vc;
 	}
 
 	void SetOrAdd(FName Name, const FMatrix& Mc)
 	{
-		MatrixConstants.FindOrAdd(Name) = Mc;
+		MatrixConstants.FindOrAdd(FNiagaraVariableInfo(Name, ENiagaraDataType::Matrix)) = Mc;
 	}
 
 	void SetOrAdd(FName Name, class FNiagaraDataObject* Cc)
 	{
-		DataConstants.FindOrAdd(Name) = Cc;
+		DataConstants.FindOrAdd(FNiagaraVariableInfo(Name, ENiagaraDataType::Curve)) = Cc;
 	}
 
 
-	float* FindScalar(FName InName){ return ScalarConstants.Find(InName); }
-	FVector4* FindVector(FName InName){ return VectorConstants.Find(InName); }
-	FMatrix* FindMatrix(FName InName){ return MatrixConstants.Find(InName); }
-	const float* FindScalar(FName InName)const { return ScalarConstants.Find(InName); }
-	const FVector4* FindVector(FName InName)const { return VectorConstants.Find(InName); }
-	const FMatrix* FindMatrix(FName InName)const { return MatrixConstants.Find(InName); }
-	FNiagaraDataObject* FindDataObj(FName InName) const 
+	float* FindScalar(FNiagaraVariableInfo ID){ return ScalarConstants.Find(ID); }
+	FVector4* FindVector(FNiagaraVariableInfo ID){ return VectorConstants.Find(ID); }
+	FMatrix* FindMatrix(FNiagaraVariableInfo ID){ return MatrixConstants.Find(ID); }
+	const float* FindScalar(FNiagaraVariableInfo ID)const { return ScalarConstants.Find(ID); }
+	const FVector4* FindVector(FNiagaraVariableInfo ID)const { return VectorConstants.Find(ID); }
+	const FMatrix* FindMatrix(FNiagaraVariableInfo ID)const { return MatrixConstants.Find(ID); }
+	FNiagaraDataObject* FindDataObj(FNiagaraVariableInfo ID) const
 	{
-		FNiagaraDataObject * const *Ptr = DataConstants.Find(InName);
+		FNiagaraDataObject * const *Ptr = DataConstants.Find(ID);
 		if (Ptr)
 		{
 			return *Ptr;
@@ -56,6 +87,10 @@ public:
 
 		return nullptr; 
 	}
+	float* FindScalar(FName Name) { return ScalarConstants.Find(FNiagaraVariableInfo(Name,ENiagaraDataType::Scalar)); }
+	FVector4* FindVector(FName Name) { return VectorConstants.Find(FNiagaraVariableInfo(Name, ENiagaraDataType::Vector)); }
+	FMatrix* FindMatrix(FName Name) { return MatrixConstants.Find(FNiagaraVariableInfo(Name, ENiagaraDataType::Matrix)); }
+	FNiagaraDataObject* FindDataObj(FName Name)const { return FindDataObj(FNiagaraVariableInfo(Name, ENiagaraDataType::Curve)); }
 
 	void Merge(FNiagaraConstantMap &InMap)
 	{
@@ -99,27 +134,37 @@ struct FNiagaraConstants
 private:
 
 	UPROPERTY()
-	TArray<float> ScalarConstants;
+	TArray<FNiagaraVariableInfo> ScalarConstantsInfo;
 	UPROPERTY()
-	TArray<FName> ScalarNames;
+	TArray<FNiagaraVariableInfo> VectorConstantsInfo;
+	UPROPERTY()
+	TArray<FNiagaraVariableInfo> MatrixConstantsInfo;
+	//UPROPERTY()
+	TArray<FNiagaraVariableInfo> DataObjectConstantsInfo;
 
 	UPROPERTY()
-	TArray<FVector4> VectorConstants;
+	TArray<float> ScalarConstants;
 	UPROPERTY()
-	TArray<FName> VectorNames;
-
+	TArray<FVector4> VectorConstants;
 	UPROPERTY()
 	TArray<FMatrix> MatrixConstants;
-	UPROPERTY()
-	TArray<FName> MatrixNames;
-
 	//UPROPERTY()
 	TArray<FNiagaraDataObject*> DataObjectConstants;
-	//UPROPERTY()
-	TArray<FName> DataObjectNames;
 
 public:
 
+	void Empty()
+	{
+		ScalarConstantsInfo.Empty();
+		VectorConstantsInfo.Empty();
+		MatrixConstantsInfo.Empty();
+		DataObjectConstantsInfo.Empty();
+		ScalarConstants.Empty();
+		VectorConstants.Empty();
+		MatrixConstants.Empty();
+		DataObjectConstants.Empty();
+	}
+
 	/** Fills the entire constants set into the constant table. */
 	void AppendToConstantsTable(TArray<FVector4>& ConstantsTable)const
 	{
@@ -134,7 +179,7 @@ public:
 		}
 
 		Idx += FMath::Min(1, NewIdx % 4);//Move to the next table entry if needed.
-		for (FVector4 Vc : VectorConstants)
+		for (const FVector4& Vc : VectorConstants)
 		{
 			ConstantsTable[Idx++] = Vc;
 		}
@@ -162,10 +207,9 @@ public:
 		int32 Idx = ConstantsTable.Num();
 		int32 NewIdx = 0;
 		ConstantsTable.AddUninitialized(ScalarTableSize() + VectorTableSize() + MatrixTableSize());
-		for (int32 i = 0; i < ScalarNames.Num(); ++i)
+		for (int32 i = 0; i < ScalarConstants.Num(); ++i)
 		{
-			FName ScName = ScalarNames[i];
-			const float* Scalar = Externals.FindScalar(ScName);
+			const float* Scalar = Externals.FindScalar(ScalarConstantsInfo[i]);
 			ConstantsTable[Idx + (NewIdx / 4)][NewIdx % 4] = Scalar ? *Scalar : ScalarConstants[i];
 			++NewIdx;
 			Idx = NewIdx % 4 == 0 ? Idx + 1 : Idx;
@@ -173,17 +217,15 @@ public:
 
 		Idx += FMath::Min(1, NewIdx % 4);//Move to the next table entry if needed.
 
-		for (int32 i = 0; i < VectorNames.Num(); ++i)
+		for (int32 i = 0; i < VectorConstants.Num(); ++i)
 		{
-			FName VcName = VectorNames[i];
-			const FVector4* Vector = Externals.FindVector(VcName);
+			const FVector4* Vector = Externals.FindVector(VectorConstantsInfo[i]);
 			ConstantsTable[Idx++] = Vector ? *Vector : VectorConstants[i];
 		}
 
-		for (int32 i = 0; i < MatrixNames.Num(); ++i)
+		for (int32 i = 0; i < MatrixConstants.Num(); ++i)
 		{
-			FName McName = MatrixNames[i];
-			const FMatrix* Mat = Externals.FindMatrix(McName);
+			const FMatrix* Mat = Externals.FindMatrix(MatrixConstantsInfo[i]);
 			FMatrix Mc = Mat ? *Mat : MatrixConstants[i];
 			ConstantsTable[Idx] = FVector4(Mc.M[0][0], Mc.M[0][1], Mc.M[0][2], Mc.M[0][3]);
 			ConstantsTable[Idx + 1] = FVector4(Mc.M[1][0], Mc.M[1][1], Mc.M[1][2], Mc.M[1][3]);
@@ -194,78 +236,73 @@ public:
 
 	void AppendExternalBufferConstants(TArray<class FNiagaraDataObject*> &DataObjs, const FNiagaraConstantMap& Externals) const
 	{
-		for (int32 i = 0; i < DataObjectNames.Num(); ++i)
+		for (int32 i = 0; i < DataObjectConstantsInfo.Num(); ++i)
 		{
-			FName CcName = DataObjectNames[i];
-			FNiagaraDataObject* Data = Externals.FindDataObj(CcName);
+			FNiagaraVariableInfo CcID = DataObjectConstantsInfo[i];
+			FNiagaraDataObject* Data = Externals.FindDataObj(CcID);
 			DataObjs.Add(Data);
 		}
 	}
 
-
-
-	void SetOrAdd(FName Name, float Sc)
+	void SetOrAdd(const FNiagaraVariableInfo& Constant, float Value)
 	{
-		int32 Idx = ScalarNames.Find(Name);
+		int32 Idx = ScalarConstantsInfo.Find(Constant);
 		if (Idx == INDEX_NONE)
 		{
-			ScalarNames.Add(Name);
-			ScalarConstants.Add(Sc);
+			ScalarConstantsInfo.Add(Constant);
+			ScalarConstants.Add(Value);
 		}
 		else
 		{
-			ScalarConstants[Idx] = Sc;
+			ScalarConstants[Idx] = Value;
 		}
-
 	}
 
-	void SetOrAdd(FName Name, const FVector4& Vc)
+	void SetOrAdd(const FNiagaraVariableInfo& Constant, FVector4 Value)
 	{
-		int32 Idx = VectorNames.Find(Name);
+		int32 Idx = VectorConstantsInfo.Find(Constant);
 		if (Idx == INDEX_NONE)
 		{
-			VectorNames.Add(Name);
-			VectorConstants.Add(Vc);
+			VectorConstantsInfo.Add(Constant);
+			VectorConstants.Add(Value);
 		}
 		else
 		{
-			VectorConstants[Idx] = Vc;
+			VectorConstants[Idx] = Value;
 		}
 	}
 
-	void SetOrAdd(FName Name, const FMatrix& Mc)
+	void SetOrAdd(const FNiagaraVariableInfo& Constant, const FMatrix& Value)
 	{
-		int32 Idx = MatrixNames.Find(Name);
+		int32 Idx = MatrixConstantsInfo.Find(Constant);
 		if (Idx == INDEX_NONE)
 		{
-			MatrixNames.Add(Name);
-			MatrixConstants.Add(Mc);
+			MatrixConstantsInfo.Add(Constant);
+			MatrixConstants.Add(Value);
 		}
 		else
 		{
-			MatrixConstants[Idx] = Mc;
+			MatrixConstants[Idx] = Value;
 		}
 	}
 
-
-
-	void SetOrAdd(FName Name, FNiagaraDataObject* Cc)
+	void SetOrAdd(const FNiagaraVariableInfo& Constant, FNiagaraDataObject* Value)
 	{
-		int32 Idx = DataObjectNames.Find(Name);
+		int32 Idx = DataObjectConstantsInfo.Find(Constant);
 		if (Idx == INDEX_NONE)
 		{
-			DataObjectNames.Add(Name);
-			DataObjectConstants.Add(Cc);
+			DataObjectConstantsInfo.Add(Constant);
+			DataObjectConstants.Add(Value);
 		}
 		else
 		{
-			DataObjectConstants[Idx] = Cc;
+			DataObjectConstants[Idx] = Value;
 		}
 	}
 
-	FORCEINLINE int32 NumScalars()const { return ScalarNames.Num(); }
-	FORCEINLINE int32 NumVectors()const { return VectorNames.Num(); }
-	FORCEINLINE int32 NumMatrices()const { return MatrixNames.Num(); }
+	FORCEINLINE int32 NumScalars()const { return ScalarConstants.Num(); }
+	FORCEINLINE int32 NumVectors()const { return VectorConstants.Num(); }
+	FORCEINLINE int32 NumMatrices()const { return MatrixConstants.Num(); }
 	FORCEINLINE int32 ScalarTableSize()const { return(NumScalars() / 4) + FMath::Min(1,NumScalars() % 4); }
 	FORCEINLINE int32 VectorTableSize()const { return(NumVectors()); }
 	FORCEINLINE int32 MatrixTableSize()const { return(NumMatrices() * 4); }
@@ -274,9 +311,9 @@ public:
 	FORCEINLINE int32 GetTableSize()const { return ScalarTableSize() + VectorTableSize() + MatrixTableSize(); }
 
 	/** Return the first constant used for this constant in a table. */
-	FORCEINLINE int32 GetTableIndex_Scalar(FName InName)const
+	FORCEINLINE int32 GetTableIndex_Scalar(const FNiagaraVariableInfo& Constant)const
 	{ 
-		int Idx = ScalarNames.Find(InName);
+		int Idx = ScalarConstantsInfo.Find(Constant);
 		if (Idx != INDEX_NONE)
 		{
 			return Idx / 4;
@@ -284,9 +321,9 @@ public:
 		return INDEX_NONE;
 	}
 
-	FORCEINLINE int32 GetTableIndex_Vector(FName InName)const
+	FORCEINLINE int32 GetTableIndex_Vector(const FNiagaraVariableInfo& Constant)const
 	{ 
-		int Idx = VectorNames.Find(InName);
+		int Idx = VectorConstantsInfo.Find(Constant);
 		if (Idx != INDEX_NONE)
 		{
 			return ScalarTableSize() + Idx;
@@ -294,9 +331,9 @@ public:
 		return INDEX_NONE;
 	}
 
-	FORCEINLINE int32 GetTableIndex_Matrix(FName InName)const
+	FORCEINLINE int32 GetTableIndex_Matrix(const FNiagaraVariableInfo& Constant)const
 	{ 
-		int Idx = MatrixNames.Find(InName);
+		int Idx = MatrixConstantsInfo.Find(Constant);
 		if (Idx != INDEX_NONE)
 		{
 			return ScalarTableSize() + VectorTableSize() + Idx * 4;
@@ -304,12 +341,12 @@ public:
 		return INDEX_NONE;
 	}
 
-	FORCEINLINE int32 GetTableIndex_DataObj(FName InName)const
+	FORCEINLINE int32 GetTableIndex_DataObj(const FNiagaraVariableInfo& Constant)const
 	{ 
-		return DataObjectNames.Find(InName); 
+		return DataObjectConstantsInfo.Find(Constant); 
 	}
 
 	/** Return the absolute index of the passed scalar. Can be used to get the component index also. */
-	FORCEINLINE int32 GetAbsoluteIndex_Scalar(FName InName)const{ return ScalarNames.Find(InName); }
+	FORCEINLINE int32 GetAbsoluteIndex_Scalar(const FNiagaraVariableInfo& Constant)const{ return ScalarConstantsInfo.Find(Constant); }
 };
 
diff --git a/Engine/Source/Runtime/Engine/Classes/Niagara/NiagaraScriptConstantData.h b/Engine/Source/Runtime/Engine/Classes/Niagara/NiagaraScriptConstantData.h
--- a/Engine/Source/Runtime/Engine/Classes/Niagara/NiagaraScriptConstantData.h
+++ b/Engine/Source/Runtime/Engine/Classes/Niagara/NiagaraScriptConstantData.h
@@ -19,6 +19,12 @@ struct FNiagaraScriptConstantData
 	UPROPERTY()
 	FNiagaraConstants InternalConstants;
 
+	void Empty()
+	{
+		ExternalConstants.Empty();
+		InternalConstants.Empty();
+	}
+
 	/** Fill a constants table ready for use in an update or spawn. */
 	void FillConstantTable(const FNiagaraConstantMap& ExternalConstantsMap, TArray<FVector4>& OutConstantTable, TArray<FNiagaraDataObject *> &DataObjTable) const
 	{
@@ -32,52 +38,23 @@ struct FNiagaraScriptConstantData
 		InternalConstants.AppendBufferConstants(DataObjTable);
 	}
 
-	void SetOrAddInternal(FName Name, float Sc)
-	{
-		InternalConstants.SetOrAdd(Name, Sc);
-	}
-
-	void SetOrAddInternal(FName Name, const FVector4& Vc)
-	{
-		InternalConstants.SetOrAdd(Name, Vc);
-	}
-
-	void SetOrAddInternal(FName Name, const FMatrix& Mc)
-	{
-		InternalConstants.SetOrAdd(Name, Mc);
-	}
-
-	void SetOrAddInternal(FName Name, FNiagaraDataObject* Curve)
-	{
-		InternalConstants.SetOrAdd(Name, Curve);
-	}
-
-	void SetOrAddExternal(FName Name, float Sc)
+	template<typename T>
+	void SetOrAddInternal(const FNiagaraVariableInfo& Constant, const T& Value)
 	{
-		ExternalConstants.SetOrAdd(Name, Sc);
+		InternalConstants.SetOrAdd(Constant, Value);
 	}
 
-	void SetOrAddExternal(FName Name, const FVector4& Vc)
+	template<typename T>
+	void SetOrAddExternal(const FNiagaraVariableInfo& Constant, const T& Value)
 	{
-		ExternalConstants.SetOrAdd(Name, Vc);
+		ExternalConstants.SetOrAdd(Constant, Value);
 	}
 
-	void SetOrAddExternal(FName Name, const FMatrix& Mc)
-	{
-		ExternalConstants.SetOrAdd(Name, Mc);
-	}
-
-	void SetOrAddExternal(FName Name, FNiagaraDataObject* Curve)
-	{
-		ExternalConstants.SetOrAdd(Name, Curve);
-	}
-
-
 	/**
 	Calculates the index into a a constant table created from this constant data for the given constant name.
 	Obviously, assumes the constant data is complete. Any additions etc will invalidate previously calculated indexes.
 	*/
-	void GetTableIndex(FName Name, bool bInternal, int32& OutConstantIdx, int32& OutComponentIndex, ENiagaraDataType& OutType)const
+	void GetTableIndex(const FNiagaraVariableInfo& Constant, bool bInternal, int32& OutConstantIdx, int32& OutComponentIndex, ENiagaraDataType& OutType)const
 	{
 		int32 Base = bInternal ? ExternalConstants.GetTableSize() : 0;
 		int32 ConstIdx = INDEX_NONE;
@@ -85,20 +62,20 @@ struct FNiagaraScriptConstantData
 		ENiagaraDataType Type;
 		const FNiagaraConstants& Consts = bInternal ? InternalConstants : ExternalConstants;
 
-		ConstIdx = Consts.GetAbsoluteIndex_Scalar(Name);
+		ConstIdx = Consts.GetAbsoluteIndex_Scalar(Constant);
 		Type = ENiagaraDataType::Scalar;
 		if (ConstIdx == INDEX_NONE)
 		{
-			ConstIdx = Consts.GetTableIndex_Vector(Name);
+			ConstIdx = Consts.GetTableIndex_Vector(Constant);
 			Type = ENiagaraDataType::Vector;
 			if (ConstIdx == INDEX_NONE)
 			{
-				ConstIdx = Consts.GetTableIndex_Matrix(Name);
+				ConstIdx = Consts.GetTableIndex_Matrix(Constant);
 				Type = ENiagaraDataType::Matrix;
 
 				if (ConstIdx == INDEX_NONE)	// curves/buffers are in a separate table, so set base to 0
 				{
-					ConstIdx = Consts.GetTableIndex_DataObj(Name);
+					ConstIdx = Consts.GetTableIndex_DataObj(Constant);
 					Type = ENiagaraDataType::Curve;
 					Base = 0;
 				}
diff --git a/Engine/Source/Runtime/Engine/Private/NiagaraComponent.cpp b/Engine/Source/Runtime/Engine/Private/NiagaraComponent.cpp
--- a/Engine/Source/Runtime/Engine/Private/NiagaraComponent.cpp
+++ b/Engine/Source/Runtime/Engine/Private/NiagaraComponent.cpp
@@ -16,7 +16,7 @@
 DECLARE_CYCLE_STAT(TEXT("Gen Verts"),STAT_NiagaraGenerateVertices,STATGROUP_Niagara);
 DECLARE_DWORD_COUNTER_STAT(TEXT("NumParticles"),STAT_NiagaraNumParticles,STATGROUP_Niagara);
 
-
+DEFINE_LOG_CATEGORY(LogNiagara);
 
 FNiagaraSceneProxy::FNiagaraSceneProxy(const UNiagaraComponent* InComponent)
 		:	FPrimitiveSceneProxy(InComponent)
@@ -136,24 +136,6 @@ void FNiagaraSceneProxy::GetDynamicMeshElements(const TArray<const FSceneView*>&
 	}
 }
 
-
-
-
-namespace ENiagaraVectorAttr
-{
-	enum Type
-	{
-		Position,
-		Velocity,
-		Color,
-		Rotation,
-		RelativeTime,
-		MaxVectorAttribs
-	};
-}
-
-
-
 //////////////////////////////////////////////////////////////////////////
 
 UNiagaraComponent::UNiagaraComponent(const FObjectInitializer& ObjectInitializer)
@@ -170,22 +152,22 @@ void UNiagaraComponent::TickComponent(float DeltaSeconds, enum ELevelTick TickTy
 
 	if (EffectInstance)
 	{ 
-
+		static FNiagaraVariableInfo Const_Zero(TEXT("ZERO"), ENiagaraDataType::Vector);
+		static FNiagaraVariableInfo Const_DeltaTime(TEXT("Delta Time"), ENiagaraDataType::Vector);
+		static FNiagaraVariableInfo Const_EmitterPos(TEXT("Emitter Position"), ENiagaraDataType::Vector);
+		static FNiagaraVariableInfo Const_EmitterX(TEXT("Emitter X Axis"), ENiagaraDataType::Vector);
+		static FNiagaraVariableInfo Const_EmitterY(TEXT("Emitter Y Axis"), ENiagaraDataType::Vector);
+		static FNiagaraVariableInfo Const_EmitterZ(TEXT("Emitter Z Axis"), ENiagaraDataType::Vector);
+		static FNiagaraVariableInfo Const_EmitterTransform(TEXT("Emitter Transform"), ENiagaraDataType::Matrix);
 		//Todo, open this up to the UI and setting via code and BPs.
-		static FName Const_Zero(TEXT("ZERO"));
 		EffectInstance->SetConstant(Const_Zero, FVector4(0.0f, 0.0f, 0.0f, 0.0f));	// zero constant
-		static FName Const_DeltaTime(TEXT("Delta Time"));
 		EffectInstance->SetConstant(Const_DeltaTime, FVector4(DeltaSeconds, DeltaSeconds, DeltaSeconds, DeltaSeconds));
-		static FName Const_EmitterPos(TEXT("Emitter Position"));
 		EffectInstance->SetConstant(Const_EmitterPos, FVector4(ComponentToWorld.GetTranslation()));
-		static FName Const_EmitterX(TEXT("Emitter X Axis"));
 		EffectInstance->SetConstant(Const_EmitterX, FVector4(ComponentToWorld.GetUnitAxis(EAxis::X)));
-		static FName Const_EmitterY(TEXT("Emitter Y Axis"));
 		EffectInstance->SetConstant(Const_EmitterY, FVector4(ComponentToWorld.GetUnitAxis(EAxis::Y)));
-		static FName Const_EmitterZ(TEXT("Emitter Z Axis"));
 		EffectInstance->SetConstant(Const_EmitterZ, FVector4(ComponentToWorld.GetUnitAxis(EAxis::Z)));
-		static FName Const_EmitterTransform(TEXT("Emitter Transform"));
 		EffectInstance->SetConstant(Const_EmitterTransform, ComponentToWorld.ToMatrixWithScale());
+
 		EffectInstance->Tick(DeltaSeconds);
 	}
 
@@ -304,3 +286,27 @@ void UNiagaraComponent::SetAsset(UNiagaraEffect *InAsset)
 
 	EffectInstance = new FNiagaraEffectInstance(Asset, this);
 }
+
+const TArray<FNiagaraVariableInfo>& UNiagaraComponent::GetSystemConstants()
+{
+	static TArray<FNiagaraVariableInfo> SystemConstants;
+	if (SystemConstants.Num() == 0)
+	{
+		static FNiagaraVariableInfo Const_Zero(TEXT("ZERO"), ENiagaraDataType::Vector);
+		static FNiagaraVariableInfo Const_DeltaTime(TEXT("Delta Time"), ENiagaraDataType::Vector);
+		static FNiagaraVariableInfo Const_EmitterPos(TEXT("Emitter Position"), ENiagaraDataType::Vector);
+		static FNiagaraVariableInfo Const_EmitterX(TEXT("Emitter X Axis"), ENiagaraDataType::Vector);
+		static FNiagaraVariableInfo Const_EmitterY(TEXT("Emitter Y Axis"), ENiagaraDataType::Vector);
+		static FNiagaraVariableInfo Const_EmitterZ(TEXT("Emitter Z Axis"), ENiagaraDataType::Vector);
+		static FNiagaraVariableInfo Const_EmitterTransform(TEXT("Emitter Transform"), ENiagaraDataType::Matrix);
+		SystemConstants.Add(Const_Zero);
+		SystemConstants.Add(Const_DeltaTime);
+		SystemConstants.Add(Const_EmitterPos);
+		SystemConstants.Add(Const_EmitterPos);
+		SystemConstants.Add(Const_EmitterX);
+		SystemConstants.Add(Const_EmitterY);
+		SystemConstants.Add(Const_EmitterZ);
+		SystemConstants.Add(Const_EmitterTransform);
+	}
+	return SystemConstants;
+}
\ No newline at end of file
diff --git a/Engine/Source/Runtime/Engine/Private/NiagaraEffectRenderer.cpp b/Engine/Source/Runtime/Engine/Private/NiagaraEffectRenderer.cpp
--- a/Engine/Source/Runtime/Engine/Private/NiagaraEffectRenderer.cpp
+++ b/Engine/Source/Runtime/Engine/Private/NiagaraEffectRenderer.cpp
@@ -202,6 +202,12 @@ void NiagaraEffectRendererSprites::GetDynamicMeshElements(const TArray<const FSc
 	CPUTimeMS += MeshElementsTimer.GetElapsedMilliseconds();
 }
 
+bool NiagaraEffectRendererSprites::SetMaterialUsage()
+{
+	//Causes deadlock :S Need to look at / rework the setting of materials and render modules.
+	//return Material && Material->CheckMaterialUsage_Concurrent(MATUSAGE_ParticleSprites);
+	return true;
+}
 
 /** Update render data buffer from attributes */
 FNiagaraDynamicDataBase *NiagaraEffectRendererSprites::GenerateVertexData(const FNiagaraEmitterParticleData &Data)
@@ -216,10 +222,17 @@ FNiagaraDynamicDataBase *NiagaraEffectRendererSprites::GenerateVertexData(const
 	RenderData.Reset(Data.GetNumParticles());
 	//CachedBounds.Init();
 
-	const FVector4 *PosPtr = Data.GetAttributeData("Position");
-	const FVector4 *ColPtr = Data.GetAttributeData("Color");
-	const FVector4 *AgePtr = Data.GetAttributeData("Age");
-	const FVector4 *RotPtr = Data.GetAttributeData("Rotation");
+	const FVector4 *PosPtr = Data.GetAttributeData(FNiagaraVariableInfo(FName(TEXT("Position")), ENiagaraDataType::Vector));
+	const FVector4 *ColPtr = Data.GetAttributeData(FNiagaraVariableInfo(FName(TEXT("Color")), ENiagaraDataType::Vector));
+	const FVector4 *AgePtr = Data.GetAttributeData(FNiagaraVariableInfo(FName(TEXT("Age")), ENiagaraDataType::Vector));
+	const FVector4 *RotPtr = Data.GetAttributeData(FNiagaraVariableInfo(FName(TEXT("Rotation")), ENiagaraDataType::Vector));
+
+	//Bail if we don't have the required attributes to render this emitter.
+	if (!PosPtr || !ColPtr || !AgePtr || !RotPtr)
+	{
+		return DynamicData;
+	}
+
 	uint32 NumSubImages = 1;
 	if (Properties)
 	{
@@ -285,7 +298,20 @@ bool NiagaraEffectRendererSprites::HasDynamicData()
 	return DynamicDataRender && DynamicDataRender->VertexData.Num() > 0;
 }
 
+const TArray<FNiagaraVariableInfo>& NiagaraEffectRendererSprites::GetRequiredAttributes()
+{
+	static TArray<FNiagaraVariableInfo> Attrs;
+	
+	if (Attrs.Num() == 0)
+	{
+		Attrs.Add(FNiagaraVariableInfo(FName(TEXT("Position")), ENiagaraDataType::Vector));
+		Attrs.Add(FNiagaraVariableInfo(FName(TEXT("Color")), ENiagaraDataType::Vector));
+		Attrs.Add(FNiagaraVariableInfo(FName(TEXT("Rotation")), ENiagaraDataType::Vector));
+		Attrs.Add(FNiagaraVariableInfo(FName(TEXT("Age")), ENiagaraDataType::Vector));
+	}
 
+	return Attrs;
+}
 
 
 
@@ -442,7 +468,27 @@ bool NiagaraEffectRendererRibbon::HasDynamicData()
 	return DynamicDataRender && DynamicDataRender->VertexData.Num() > 0;
 }
 
+const TArray<FNiagaraVariableInfo>& NiagaraEffectRendererRibbon::GetRequiredAttributes()
+{
+	static TArray<FNiagaraVariableInfo> Attrs;
+
+	if (Attrs.Num() == 0)
+	{
+		Attrs.Add(FNiagaraVariableInfo(FName(TEXT("Position")), ENiagaraDataType::Vector));
+		Attrs.Add(FNiagaraVariableInfo(FName(TEXT("Color")), ENiagaraDataType::Vector));
+		Attrs.Add(FNiagaraVariableInfo(FName(TEXT("Rotation")), ENiagaraDataType::Vector));
+		Attrs.Add(FNiagaraVariableInfo(FName(TEXT("Age")), ENiagaraDataType::Vector));
+	}
+
+	return Attrs;
+}
 
+bool NiagaraEffectRendererRibbon::SetMaterialUsage()
+{
+	//Causes deadlock :S Need to look at / rework the setting of materials and render modules.
+	//return Material && Material->CheckMaterialUsage_Concurrent(MATUSAGE_BeamTrails);
+	return true;
+}
 
 FNiagaraDynamicDataBase *NiagaraEffectRendererRibbon::GenerateVertexData(const FNiagaraEmitterParticleData &Data)
 {
@@ -464,21 +510,26 @@ FNiagaraDynamicDataBase *NiagaraEffectRendererRibbon::GenerateVertexData(const F
 		SortedIndices.Add(Idx);
 	}
 
-	const FVector4 *AgeData = Data.GetAttributeData("Age");
+	const FVector4 *PosPtr = Data.GetAttributeData(FNiagaraVariableInfo(FName(TEXT("Position")), ENiagaraDataType::Vector));
+	const FVector4 *ColorPtr = Data.GetAttributeData(FNiagaraVariableInfo(FName(TEXT("Color")), ENiagaraDataType::Vector));
+	const FVector4 *RotPtr = Data.GetAttributeData(FNiagaraVariableInfo(FName(TEXT("Rotation")), ENiagaraDataType::Vector));
+	const FVector4 *AgePtr = Data.GetAttributeData(FNiagaraVariableInfo(FName(TEXT("Age")), ENiagaraDataType::Vector));
+
+	//Bail if we don't have the required attributes to render this emitter.
+	if (!PosPtr || !ColorPtr || !AgePtr || !RotPtr)
+	{
+		return DynamicData;
+	}
+
 	SortedIndices.Sort(
-		[&AgeData](const int32& A, const int32& B) {
-		return AgeData[A].X < AgeData[B].X;
+		[&AgePtr](const int32& A, const int32& B) {
+		return AgePtr[A].X < AgePtr[B].X;
 	}
 	);
 
 
 	FVector2D UVs[4] = { FVector2D(0.0f, 0.0f), FVector2D(1.0f, 0.0f), FVector2D(1.0f, 1.0f), FVector2D(0.0f, 1.0f) };
 
-	const FVector4 *PosPtr = Data.GetAttributeData("Position");
-	const FVector4 *ColorPtr = Data.GetAttributeData("Color");
-	const FVector4 *AgePtr = Data.GetAttributeData("Age");
-	const FVector4 *RotPtr = Data.GetAttributeData("Rotation");
-
 	FVector PrevPos, PrevPos2, PrevDir(0.0f, 0.0f, 0.1f);
 	for (int32 i = 0; i < SortedIndices.Num() - 1; i++)
 	{
@@ -519,4 +570,4 @@ FNiagaraDynamicDataBase *NiagaraEffectRendererRibbon::GenerateVertexData(const F
 	CPUTimeMS = VertexDataTimer.GetElapsedMilliseconds();
 
 	return DynamicData;
-}
+}
\ No newline at end of file
diff --git a/Engine/Source/Runtime/Engine/Private/NiagaraSimulation.cpp b/Engine/Source/Runtime/Engine/Private/NiagaraSimulation.cpp
--- a/Engine/Source/Runtime/Engine/Private/NiagaraSimulation.cpp
+++ b/Engine/Source/Runtime/Engine/Private/NiagaraSimulation.cpp
@@ -14,7 +14,8 @@ FNiagaraSimulation::FNiagaraSimulation(FNiagaraEmitterProperties *InProps)
 , EffectRenderer(nullptr)
 {
 	Props = InProps;
-	Constants.Merge(InProps->ExternalConstants);
+
+	Init();
 }
 
 FNiagaraSimulation::FNiagaraSimulation(FNiagaraEmitterProperties *InProps, ERHIFeatureLevel::Type InFeatureLevel)
@@ -25,8 +26,79 @@ FNiagaraSimulation::FNiagaraSimulation(FNiagaraEmitterProperties *InProps, ERHIF
 	, EffectRenderer(nullptr)
 {
 	Props = InProps;
-	Constants.Merge(InProps->ExternalConstants);
-	SetRenderModuleType(InProps->RenderModuleType, InFeatureLevel);
+
+	Init();
+	SetRenderModuleType(Props->RenderModuleType, InFeatureLevel);
+}
+
+void FNiagaraSimulation::Init()
+{
+	if (Props && Props->UpdateScript && Props->SpawnScript)
+	{
+		Data.SetAttributes(Props->UpdateScript->Attributes);
+		Constants.Empty();
+		Constants.Merge(Props->ExternalConstants);
+
+		bool bDisable = false;
+		//Temporarily enforcing that spawn scripts have exactly the same attriubtes as update scripts.
+		//To fix this we need to refactor the scripts into source and runtime objects.
+		//The runtime will be a script autogenerated to be a safe combination of all scripts attributes. 
+		//With an options somewhere for handling of attributes not set in a particular script. 
+		//eg. Zeroing for spawn of attributes in update script that are not set explicitly in the spawn script. And pass-through copy for the reverse.
+		if (Props->SpawnScript->Attributes.Num() == Props->UpdateScript->Attributes.Num())
+		{
+			for (int32 i = 0; i < Props->SpawnScript->Attributes.Num(); ++i)
+			{
+				if (Props->SpawnScript->Attributes[i] != Props->UpdateScript->Attributes[i])
+				{
+					bDisable = true;
+					break;
+				}
+			}
+		}
+		else
+		{
+			bDisable = true;
+		}
+
+		if (bDisable)
+		{
+			Data.Reset();//Clear the attributes to mark the sim as disabled independatly of the user set Enabled flag.
+			bIsEnabled = false;//Also force the user flag to give an indication to the user.
+			UE_LOG(LogNiagara, Error, TEXT("This emitter cannot be enabled because it's spawn script and update script don't have matching attributes."));
+		}
+	}
+	else
+	{
+		Data.Reset();//Clear the attributes to mark the sim as disabled independatly of the user set Enabled flag.
+		bIsEnabled = false;//Also force the user flag to give an indication to the user.
+		//TODO - Arbitrary named scripts. Would need some base functionality for Spawn/Udpate to be called that can be overriden in BPs for emitters with custom scripts.
+		UE_LOG(LogNiagara, Error, TEXT("This emitter cannot be enabled because it's doesn't have both an update and spawn script."));
+	}
+}
+
+void FNiagaraSimulation::SetProperties(FNiagaraEmitterProperties *InProps)	
+{ 
+	Props = InProps; 
+	Init();
+}
+
+void FNiagaraSimulation::SetSpawnScript(UNiagaraScript* Script)
+{
+	Props->SpawnScript = Script;
+	Init();
+}
+
+void FNiagaraSimulation::SetUpdateScript(UNiagaraScript* Script)
+{
+	Props->UpdateScript = Script;
+	Init();
+}
+
+void FNiagaraSimulation::SetEnabled(bool bEnabled)
+{
+	bIsEnabled = bEnabled;
+	Init();
 }
 
 
@@ -53,6 +125,9 @@ void FNiagaraSimulation::Tick(float DeltaSeconds)
 
 	SimpleTimer TickTime;
 
+	if (Data.GetNumAttributes() == 0)
+		return;
+
 	// Cache the ComponentToWorld transform.
 //	CachedComponentToWorld = Component.GetComponentToWorld();
 
@@ -105,6 +180,9 @@ void FNiagaraSimulation::UpdateParticles(
 	int32 NumParticles
 	)
 {
+	if (Data.GetNumAttributes() == 0)
+		return;
+
 	SCOPE_CYCLE_COUNTER(STAT_NiagaraSimulate);
 	Age += DeltaSeconds;
 	Constants.SetOrAdd(TEXT("Emitter Age"), FVector4(Age, Age, Age, Age));
@@ -148,8 +226,7 @@ int32 FNiagaraSimulation::SpawnAndKillParticles(int32 NumToSpawn)
 {
 	SCOPE_CYCLE_COUNTER(STAT_NiagaraSpawnAndKill);
 	int32 OrigNumParticles = Data.GetNumParticles();
-	int32 CurNumParticles = OrigNumParticles;
-	CurNumParticles = SpawnParticles(NumToSpawn);
+	int32 CurNumParticles = OrigNumParticles + NumToSpawn;
 
 	// run the spawn graph over all new particles
 	if (Props->SpawnScript && Props->SpawnScript->ByteCode.Num())
@@ -188,23 +265,47 @@ int32 FNiagaraSimulation::SpawnAndKillParticles(int32 NumToSpawn)
 			);
 	}
 
+	//TODO: NEED TO HANDLE PARTICLE DEATH WITHOUT RELYING ON AGE ATTRIBUTE. UNLESS WE FORCE IT.
+	//MAYBE WITH BRANCHING IN VM AND A KILL() NODE.
+
 	// Iterate over looking for dead particles and move from the end of the list to the dead location, compacting in the process
 	int32 ParticleIndex = 0;
-	const FVector4* ParticleRelativeTimes = Data.GetAttributeData("Age");
-	while (ParticleIndex < OrigNumParticles)
+	const FVector4* ParticleRelativeTimes = Data.GetAttributeData(FNiagaraVariableInfo(FName(TEXT("Age")), ENiagaraDataType::Vector));
+	if (ParticleRelativeTimes)
 	{
-		if (ParticleRelativeTimes[ParticleIndex].X > 1.0f)
+		while (ParticleIndex < OrigNumParticles)
 		{
-			// Particle is dead, move one from the end here.
-			MoveParticleToIndex(--CurNumParticles, ParticleIndex);
+			if (ParticleRelativeTimes[ParticleIndex].X > 1.0f)
+			{
+				// Particle is dead, move one from the end here.
+				MoveParticleToIndex(--CurNumParticles, ParticleIndex);
+			}
+			ParticleIndex++;
 		}
-		ParticleIndex++;
 	}
 
 	Data.SetNumParticles(CurNumParticles);
 	return CurNumParticles;
 }
 
+bool FNiagaraSimulation::CheckAttriubtesForRenderer()
+{
+	bool bOk = true;
+	if (EffectRenderer)
+	{
+		const TArray<FNiagaraVariableInfo>& RequiredAttrs = EffectRenderer->GetRequiredAttributes();
+
+		for (auto& Attr : RequiredAttrs)
+		{
+			if (!Data.HasAttriubte(Attr))
+			{
+				bOk = false;
+				UE_LOG(LogNiagara, Error, TEXT("Cannot render %s because it does not define attriubte %s."), *Props->Name, *Attr.Name.ToString());
+			}
+		}
+	}
+	return bOk;
+}
 
 /** Replace the current effect renderer with a new one of Type, transferring the existing material over
  *	to the new renderer. Don't forget to call RenderModuleUpdate on the SceneProxy after calling this! 
@@ -241,34 +342,6 @@ void FNiagaraSimulation::SetRenderModuleType(EEmitterRenderModuleType Type, ERHI
 		}
 
 		EffectRenderer->SetMaterial(Material, FeatureLevel);
+		CheckAttriubtesForRenderer();
 	}
 }
-
-
-
-int32 FNiagaraSimulation::SpawnParticles(int32 NumToSpawn)
-{
-	FVector4 *PosPtr = Data.GetAttributeDataWrite("Position");
-	FVector4 *VelPtr = Data.GetAttributeDataWrite("Velocity");
-	FVector4 *ColPtr = Data.GetAttributeDataWrite("Color");
-	FVector4 *RotPtr = Data.GetAttributeDataWrite("Rotation");
-	FVector4 *AgePtr = Data.GetAttributeDataWrite("Age");
-
-	// Spawn new Particles at the end of the buffer
-	int32 ParticleIndex = Data.GetNumParticles();
-	FVector SpawnLocation = *Constants.FindVector(TEXT("Emitter Position"));
-	for (int32 i = 0; i < NumToSpawn; i++)
-	{
-		SpawnLocation.X += FMath::FRandRange(-20.f, 20.f);
-		SpawnLocation.Y += FMath::FRandRange(-20.f, 20.f);
-		SpawnLocation.Z += FMath::FRandRange(-20.f, 20.f);
-
-		PosPtr[ParticleIndex] = SpawnLocation;
-		ColPtr[ParticleIndex] = FVector4(0.0f, 0.0f, 0.0f, 0.0f);
-		VelPtr[ParticleIndex] = FVector4(0.0f, 0.0f, 2.0f, 0.0f);
-		RotPtr[ParticleIndex] = FVector4(0.0f, 0.0f, 0.0f, 0.0f);
-		AgePtr[ParticleIndex] = FVector4(0.0f, 0.0f, 0.0f, 0.0f);
-		ParticleIndex++;
-	}
-	return ParticleIndex;
-}
\ No newline at end of file
