Niagara; first pass buffer write, curve data object, curve data editor in the effect editor

diff --git a/Engine/Source/Editor/NiagaraEditor/Private/NiagaraCompiler_VectorVM.cpp b/Engine/Source/Editor/NiagaraEditor/Private/NiagaraCompiler_VectorVM.cpp
--- a/Engine/Source/Editor/NiagaraEditor/Private/NiagaraCompiler_VectorVM.cpp
+++ b/Engine/Source/Editor/NiagaraEditor/Private/NiagaraCompiler_VectorVM.cpp
@@ -599,4 +599,14 @@ void FNiagaraCompiler_VectorVM::Sample_Internal(TArray<TNiagaraExprPtr>& InputEx
 	OutputExpressions.Add(Expression_VMNative(VectorVM::EOp::sample, InputExpressions));
 }
 
+void FNiagaraCompiler_VectorVM::Write_Internal(TArray<TNiagaraExprPtr>& InputExpressions, TArray<TNiagaraExprPtr>& OutputExpressions)
+{
+	OutputExpressions.Add(Expression_VMNative(VectorVM::EOp::bufferwrite, InputExpressions));
+}
+
+void FNiagaraCompiler_VectorVM::EventBroadcast_Internal(TArray<TNiagaraExprPtr>& InputExpressions, TArray<TNiagaraExprPtr>& OutputExpressions)
+{
+	OutputExpressions.Add(Expression_VMNative(VectorVM::EOp::eventbroadcast, InputExpressions));
+}
+
 #undef LOCTEXT_NAMESPACE
diff --git a/Engine/Source/Editor/NiagaraEditor/Private/SNiagaraEffectEditorWidget.cpp b/Engine/Source/Editor/NiagaraEditor/Private/SNiagaraEffectEditorWidget.cpp
--- a/Engine/Source/Editor/NiagaraEditor/Private/SNiagaraEffectEditorWidget.cpp
+++ b/Engine/Source/Editor/NiagaraEditor/Private/SNiagaraEffectEditorWidget.cpp
@@ -4,6 +4,7 @@
 #include "SNiagaraEffectEditorWidget.h"
 #include "SNotificationList.h"
 #include "SExpandableArea.h"
+#include "SSplitter.h"
 #include "Editor/PropertyEditor/Public/PropertyEditorModule.h"
 #include "Editor/PropertyEditor/Public/IDetailsView.h"
 
@@ -14,51 +15,6 @@ BEGIN_SLATE_FUNCTION_BUILD_OPTIMIZATION
 
 void SNiagaraEffectEditorWidget::Construct(const FArguments& InArgs)
 {
-	/*
-	Commands = MakeShareable(new FUICommandList());
-	IsEditable = InArgs._IsEditable;
-	Appearance = InArgs._Appearance;
-	TitleBarEnabledOnly = InArgs._TitleBarEnabledOnly;
-	TitleBar = InArgs._TitleBar;
-	bAutoExpandActionMenu = InArgs._AutoExpandActionMenu;
-	ShowGraphStateOverlay = InArgs._ShowGraphStateOverlay;
-	OnNavigateHistoryBack = InArgs._OnNavigateHistoryBack;
-	OnNavigateHistoryForward = InArgs._OnNavigateHistoryForward;
-	OnNodeSpawnedByKeymap = InArgs._GraphEvents.OnNodeSpawnedByKeymap;
-	*/
-
-	/*
-	// Make sure that the editor knows about what kinds
-	// of commands GraphEditor can do.
-	FGraphEditorCommands::Register();
-
-	// Tell GraphEditor how to handle all the known commands
-	{
-	Commands->MapAction(FGraphEditorCommands::Get().ReconstructNodes,
-	FExecuteAction::CreateSP(this, &SGraphEditorImpl::ReconstructNodes),
-	FCanExecuteAction::CreateSP(this, &SGraphEditorImpl::CanReconstructNodes)
-	);
-
-	Commands->MapAction(FGraphEditorCommands::Get().BreakNodeLinks,
-	FExecuteAction::CreateSP(this, &SGraphEditorImpl::BreakNodeLinks),
-	FCanExecuteAction::CreateSP(this, &SGraphEditorImpl::CanBreakNodeLinks)
-	);
-
-	Commands->MapAction(FGraphEditorCommands::Get().BreakPinLinks,
-	FExecuteAction::CreateSP(this, &SGraphEditorImpl::BreakPinLinks, true),
-	FCanExecuteAction::CreateSP(this, &SGraphEditorImpl::CanBreakPinLinks)
-	);
-
-	// Append any additional commands that a consumer of GraphEditor wants us to be aware of.
-	const TSharedPtr<FUICommandList>& AdditionalCommands = InArgs._AdditionalCommands;
-	if (AdditionalCommands.IsValid())
-	{
-	Commands->Append(AdditionalCommands.ToSharedRef());
-	}
-
-	}
-	*/
-
 	EffectObj = InArgs._EffectObj;
 	if (InArgs._EffectInstance)
 	{
@@ -78,71 +34,98 @@ void SNiagaraEffectEditorWidget::Construct(const FArguments& InArgs)
 	TSharedPtr<SOverlay> OverlayWidget;
 	this->ChildSlot
 		[
-			SNew(SHorizontalBox)
-			+ SHorizontalBox::Slot()
-			.AutoWidth()
+			SNew(SSplitter)
+			.Orientation(Orient_Vertical)
+			+ SSplitter::Slot()
+			.Value(0.75f)
 			[
-				SNew(SVerticalBox)
-				+ SVerticalBox::Slot()
-				.AutoHeight()
+
+			SNew(SVerticalBox)
+			+ SVerticalBox::Slot()
+			.AutoHeight()
+			[
+
+				SNew(SHorizontalBox)
+				+ SHorizontalBox::Slot()
+				.AutoWidth()
 				[
-					SNew(SBox)
-					.WidthOverride(512)
-					.HeightOverride(512)
+					SNew(SVerticalBox)
+					+ SVerticalBox::Slot()
+					.AutoHeight()
 					[
-						Viewport.ToSharedRef()
+						SNew(SBox)
+						.WidthOverride(512)
+						.HeightOverride(512)
+						[
+							Viewport.ToSharedRef()
+						]
 					]
 				]
-			]
-			+ SHorizontalBox::Slot()
-				.AutoWidth()
-				[
-					SAssignNew(OverlayWidget, SOverlay)
-
-					+SOverlay::Slot()
-					.Padding(10)
-					.VAlign(VAlign_Fill)
+				+ SHorizontalBox::Slot()
+					.AutoWidth()
 					[
-						SAssignNew(ListView, SListView<TSharedPtr<FNiagaraSimulation> >)
-						.ItemHeight(256)
-						.ListItemsSource(&(this->EffectInstance->GetEmitters()))
-						.OnGenerateRow(this, &SNiagaraEffectEditorWidget::OnGenerateWidgetForList)
-						.SelectionMode(ESelectionMode::None)
-					]
+						SAssignNew(OverlayWidget, SOverlay)
 
-					// Bottom-right corner text indicating the type of tool
-					+ SOverlay::Slot()
+						+SOverlay::Slot()
 						.Padding(10)
-						.VAlign(VAlign_Bottom)
-						.HAlign(HAlign_Right)
+						.VAlign(VAlign_Fill)
 						[
-							SNew(STextBlock)
-							.Visibility(EVisibility::HitTestInvisible)
-							.TextStyle(FEditorStyle::Get(), "Graph.CornerText")
-							.Text(LOCTEXT("NiagaraEditorLabel", "Niagara Effect"))
+							SAssignNew(ListView, SListView<TSharedPtr<FNiagaraSimulation> >)
+							.ItemHeight(256)
+							.ListItemsSource(&(this->EffectInstance->GetEmitters()))
+							.OnGenerateRow(this, &SNiagaraEffectEditorWidget::OnGenerateWidgetForList)
+							.OnSelectionChanged(this, &SNiagaraEffectEditorWidget::OnSelectionChanged)
+							//.SelectionMode(ESelectionMode::None)
 						]
 
-					// Top-right corner text indicating read only when not simulating
-					+ SOverlay::Slot()
-						.Padding(20)
-						.VAlign(VAlign_Top)
-						.HAlign(HAlign_Right)
-						[
-							SNew(STextBlock)
-							.Visibility(this, &SNiagaraEffectEditorWidget::ReadOnlyVisibility)
-							.TextStyle(FEditorStyle::Get(), "Graph.CornerText")
-							.Text(LOCTEXT("ReadOnlyLabel", "Read Only"))
-						]
+						// Bottom-right corner text indicating the type of tool
+						+ SOverlay::Slot()
+							.Padding(10)
+							.VAlign(VAlign_Bottom)
+							.HAlign(HAlign_Right)
+							[
+								SNew(STextBlock)
+								.Visibility(EVisibility::HitTestInvisible)
+								.TextStyle(FEditorStyle::Get(), "Graph.CornerText")
+								.Text(LOCTEXT("NiagaraEditorLabel", "Niagara Effect"))
+							]
+
+						// Top-right corner text indicating read only when not simulating
+						+ SOverlay::Slot()
+							.Padding(20)
+							.VAlign(VAlign_Top)
+							.HAlign(HAlign_Right)
+							[
+								SNew(STextBlock)
+								.Visibility(this, &SNiagaraEffectEditorWidget::ReadOnlyVisibility)
+								.TextStyle(FEditorStyle::Get(), "Graph.CornerText")
+								.Text(LOCTEXT("ReadOnlyLabel", "Read Only"))
+							]
+
+						// Bottom-right corner text for notification list position
+						+ SOverlay::Slot()
+							.Padding(15.f)
+							.VAlign(VAlign_Bottom)
+							.HAlign(HAlign_Right)
+							[
+								SAssignNew(NotificationListPtr, SNotificationList)
+								.Visibility(EVisibility::HitTestInvisible)
+							]
+					]
+			]
 
-					// Bottom-right corner text for notification list position
-					+ SOverlay::Slot()
-						.Padding(15.f)
-						.VAlign(VAlign_Bottom)
-						.HAlign(HAlign_Right)
+			]
+			+ SSplitter::Slot()
+			.Value(0.25f)
+				[
+					NGED_SECTION_BORDER
+					[
+						SNew(SBox)
+						.MinDesiredHeight(400).MinDesiredWidth(1920)
 						[
-							SAssignNew(NotificationListPtr, SNotificationList)
-							.Visibility(EVisibility::HitTestInvisible)
+							SAssignNew(TimeLine, SNiagaraTimeline)
 						]
+					]
 				]
 		];
 
@@ -150,6 +133,18 @@ void SNiagaraEffectEditorWidget::Construct(const FArguments& InArgs)
 }
 
 
+void SNiagaraTimeline::Construct(const FArguments& InArgs)
+{
+	Emitter = InArgs._Emitter;
+	EffectInstance = InArgs._EffectInstance;
+	Effect = InArgs._Effect;
+
+	TSharedPtr<SOverlay> OverlayWidget;
+	this->ChildSlot
+	[
+		SAssignNew(CurveEditor, SCurveEditor)
+	];
+}
 
 
 void SEmitterWidget::Construct(const FArguments& InArgs)
diff --git a/Engine/Source/Editor/NiagaraEditor/Private/SNiagaraEffectEditorWidget.h b/Engine/Source/Editor/NiagaraEditor/Private/SNiagaraEffectEditorWidget.h
--- a/Engine/Source/Editor/NiagaraEditor/Private/SNiagaraEffectEditorWidget.h
+++ b/Engine/Source/Editor/NiagaraEditor/Private/SNiagaraEffectEditorWidget.h
@@ -8,6 +8,7 @@
 #include "WorkflowOrientedApp/SContentReference.h"
 #include "SCheckBox.h"
 #include "SNameComboBox.h"
+#include "SCurveEditor.h"
 #include "AssetThumbnail.h"
 #include "SNiagaraEffectEditorViewport.h"
 #include "Components/NiagaraComponent.h"
@@ -221,7 +222,7 @@ public:
 		this->ChildSlot
 		[
 			SNew(SHorizontalBox)
-			+ SHorizontalBox::Slot().VAlign(VAlign_Center).Padding(4)[SNew(STextBlock).Text(FText::FromName(ConstantName))]
+			+ SHorizontalBox::Slot().VAlign(VAlign_Center).AutoWidth().Padding(4)[SNew(STextBlock).Text(FText::FromName(ConstantName))]
 			+ SHorizontalBox::Slot().VAlign(VAlign_Center).Padding(4)
 			[
 				SNew(SContentReference)
@@ -388,6 +389,35 @@ public:
 };
 
 
+class SNiagaraTimeline : public SCompoundWidget, public FNotifyHook
+{
+private:
+	TSharedPtr<FNiagaraSimulation> Emitter;
+	UNiagaraEffect *Effect;
+	FNiagaraEffectInstance *EffectInstance;
+	TSharedPtr<SCurveEditor>CurveEditor;
+
+public:
+	SLATE_BEGIN_ARGS(SNiagaraTimeline) :
+	_Emitter(nullptr)
+	{
+	}
+	SLATE_ARGUMENT(TSharedPtr<FNiagaraSimulation>, Emitter)
+		SLATE_ARGUMENT(FNiagaraEffectInstance*, EffectInstance)
+		SLATE_ARGUMENT(UNiagaraEffect*, Effect)
+		SLATE_END_ARGS()
+
+	virtual void Construct(const FArguments& InArgs);
+
+	void SetCurve(UCurveBase *Curve)
+	{
+		CurveEditor->SetCurveOwner(Curve);
+	}
+
+};
+
+
+
 
 
 class SNiagaraEffectEditorWidget : public SCompoundWidget, public FNotifyHook
@@ -403,6 +433,7 @@ private:
 	
 	TSharedPtr< SListView<TSharedPtr<FNiagaraSimulation> > > ListView;
 
+	TSharedPtr< SNiagaraTimeline > TimeLine;
 public:
 	SLATE_BEGIN_ARGS(SNiagaraEffectEditorWidget)
 		: _EffectObj(nullptr)
@@ -441,7 +472,17 @@ public:
 			];
 	}
 
-
+	void OnSelectionChanged( TSharedPtr<FNiagaraSimulation> SelectedItem, ESelectInfo::Type SelType)
+	{
+		if ( SelectedItem.Get() != nullptr )
+		{
+			if (SelectedItem->GetProperties()->ExternalConstants.GetDataConstants().Num())
+			{
+				FNiagaraDataObject *DataObj = SelectedItem->GetProperties()->ExternalConstants.GetDataConstants().CreateConstIterator().Value();
+				TimeLine->SetCurve(static_cast<FNiagaraCurveDataObject*>(DataObj)->GetCurveObject());
+			}
+		}
+	}
 
 	void Construct(const FArguments& InArgs);
 
diff --git a/Engine/Source/Editor/UnrealEd/Private/Niagara/NiagaraEditorCommon.cpp b/Engine/Source/Editor/UnrealEd/Private/Niagara/NiagaraEditorCommon.cpp
--- a/Engine/Source/Editor/UnrealEd/Private/Niagara/NiagaraEditorCommon.cpp
+++ b/Engine/Source/Editor/UnrealEd/Private/Niagara/NiagaraEditorCommon.cpp
@@ -528,6 +528,25 @@ void FNiagaraOpInfo::Init()
 	Op->Outputs.Add(FNiagaraOpInOutInfo(Result, ENiagaraDataType::Vector, ResultText, ResultText, DefaultStr_VecOne));
 	Op->OpDelegate.BindStatic(&INiagaraCompiler::Sample);
 
+	Op = &OpInfos.Add(Write);
+	Op->Name = Write;
+	Op->FriendlyName = NSLOCTEXT("NiagaraOpInfo", "Write Name", "Buffer Write");
+	Op->Description = NSLOCTEXT("NiagaraOpInfo", "Write Desc", "Result = Write Value to buffer at index Coords [0;1]^4");
+	Op->Inputs.Add(FNiagaraOpInOutInfo(A, ENiagaraDataType::Curve, NSLOCTEXT("NiagaraOpInfo", "Curve", "Curve"), NSLOCTEXT("NiagaraOpInfo", "Curve", "Curve"), DefaultStr_VecOne));
+	Op->Inputs.Add(FNiagaraOpInOutInfo(B, ENiagaraDataType::Vector, NSLOCTEXT("NiagaraOpInfo", "Coords", "Coords"), NSLOCTEXT("NiagaraOpInfo", "Coords", "Coords"), DefaultStr_VecZero));
+	Op->Inputs.Add(FNiagaraOpInOutInfo(B, ENiagaraDataType::Vector, NSLOCTEXT("NiagaraOpInfo", "Value", "Value"), NSLOCTEXT("NiagaraOpInfo", "Value", "Value"), DefaultStr_VecZero));
+	Op->Outputs.Add(FNiagaraOpInOutInfo(Result, ENiagaraDataType::Vector, ResultText, ResultText, DefaultStr_VecOne));
+	Op->OpDelegate.BindStatic(&INiagaraCompiler::Write);
+
+	Op = &OpInfos.Add(EventBroadcast);
+	Op->Name = EventBroadcast;
+	Op->FriendlyName = NSLOCTEXT("NiagaraOpInfo", "EventBroadcast Name", "Broadcast Event");
+	Op->Description = NSLOCTEXT("NiagaraOpInfo", "EventBroadcast Desc", "Broadcast Event if trigger is >=1.0");
+	Op->Inputs.Add(FNiagaraOpInOutInfo(A, ENiagaraDataType::Curve, NSLOCTEXT("NiagaraOpInfo", "Trigger", "Trigger"), NSLOCTEXT("NiagaraOpInfo", "Trigger", "Trigger"), DefaultStr_VecOne));
+//	Op->Outputs.Add(FNiagaraOpInOutInfo(Result, ENiagaraDataType::Vector, ResultText, ResultText, DefaultStr_VecOne));
+	Op->OpDelegate.BindStatic(&INiagaraCompiler::EventBroadcast);
+
+
 	#define NiagaraOp(OpName) if (GetOpInfo(OpName) == NULL)\
 	{\
 	UE_LOG(LogNiagaraCompiler, Fatal, TEXT("Couldn't find info for \"%s\". Have you added it to the NiagaraOpList but not provided the extra information in FNiagaraOpInfo::Init() ?"), TEXT(#OpName)); \
diff --git a/Engine/Source/Editor/UnrealEd/Public/NiagaraEditorCommon.h b/Engine/Source/Editor/UnrealEd/Public/NiagaraEditorCommon.h
--- a/Engine/Source/Editor/UnrealEd/Public/NiagaraEditorCommon.h
+++ b/Engine/Source/Editor/UnrealEd/Public/NiagaraEditorCommon.h
@@ -63,6 +63,8 @@
 	NiagaraOp(Inverse)\
 	NiagaraOp(LessThan)\
 	NiagaraOp(Sample)\
+	NiagaraOp(Write)\
+	NiagaraOp(EventBroadcast)\
 
 enum class ENiagaraExpressionResultLocation
 {
diff --git a/Engine/Source/Runtime/Engine/Classes/Engine/NiagaraEffect.h b/Engine/Source/Runtime/Engine/Classes/Engine/NiagaraEffect.h
--- a/Engine/Source/Runtime/Engine/Classes/Engine/NiagaraEffect.h
+++ b/Engine/Source/Runtime/Engine/Classes/Engine/NiagaraEffect.h
@@ -132,13 +132,20 @@ public:
 
 	void Tick(float DeltaSeconds)
 	{
+		Constants.SetOrAdd(FName(TEXT("EffectGrid")), &VolumeGrid); 
+
 		// pass the constants down to the emitter
 		// TODO: should probably just pass a pointer to the table
+		FBox EffectBounds;
+		EffectBounds.Init();
+
 		for (TSharedPtr<FNiagaraSimulation>&it : Emitters)
 		{
 			it->SetConstants(Constants);
 			it->GetConstants().Merge(it->GetProperties()->ExternalConstants);
 			it->Tick(DeltaSeconds);
+
+			EffectBounds += it->GetEffectRenderer()->GetBounds();
 		}
 	}
 
@@ -169,5 +176,6 @@ private:
 
 	TArray< TSharedPtr<FNiagaraSimulation> > Emitters;
 
+	FNiagaraSparseVolumeDataObject VolumeGrid;
 
 };
\ No newline at end of file
diff --git a/Engine/Source/Runtime/Engine/Classes/Engine/NiagaraEffectRenderer.h b/Engine/Source/Runtime/Engine/Classes/Engine/NiagaraEffectRenderer.h
--- a/Engine/Source/Runtime/Engine/Classes/Engine/NiagaraEffectRenderer.h
+++ b/Engine/Source/Runtime/Engine/Classes/Engine/NiagaraEffectRenderer.h
@@ -101,7 +101,7 @@ public:
 	virtual UClass *GetPropertiesClass() = 0;
 
 	float GetCPUTimeMS() { return CPUTimeMS; }
-
+	FBox GetBounds()	{ return CachedBounds;  }
 protected:
 	mutable float CPUTimeMS;
 
@@ -111,6 +111,7 @@ protected:
 	}
 
 	UMaterial* Material;
+	FBox CachedBounds;
 
 private:
 	FMaterialRelevance MaterialRelevance;
diff --git a/Engine/Source/Runtime/Engine/Classes/Niagara/NiagaraConstantSet.h b/Engine/Source/Runtime/Engine/Classes/Niagara/NiagaraConstantSet.h
--- a/Engine/Source/Runtime/Engine/Classes/Niagara/NiagaraConstantSet.h
+++ b/Engine/Source/Runtime/Engine/Classes/Niagara/NiagaraConstantSet.h
@@ -115,6 +115,11 @@ public:
 		}
 	}
 
+	const TMap<FNiagaraVariableInfo, class FNiagaraDataObject*> &GetDataConstants()	const
+	{ 
+		return DataConstants; 
+	}
+
 
 	virtual bool Serialize(FArchive &Ar)
 	{
diff --git a/Engine/Source/Runtime/Engine/Private/Components/NiagaraComponent.cpp b/Engine/Source/Runtime/Engine/Private/Components/NiagaraComponent.cpp
--- a/Engine/Source/Runtime/Engine/Private/Components/NiagaraComponent.cpp
+++ b/Engine/Source/Runtime/Engine/Private/Components/NiagaraComponent.cpp
@@ -155,6 +155,7 @@ void UNiagaraComponent::TickComponent(float DeltaSeconds, enum ELevelTick TickTy
 		static FNiagaraVariableInfo Const_Zero(TEXT("ZERO"), ENiagaraDataType::Vector);
 		static FNiagaraVariableInfo Const_DeltaTime(TEXT("Delta Time"), ENiagaraDataType::Vector);
 		static FNiagaraVariableInfo Const_EmitterPos(TEXT("Emitter Position"), ENiagaraDataType::Vector);
+		static FNiagaraVariableInfo Const_EmitterAge(TEXT("Emitter Age"), ENiagaraDataType::Vector);
 		static FNiagaraVariableInfo Const_EmitterX(TEXT("Emitter X Axis"), ENiagaraDataType::Vector);
 		static FNiagaraVariableInfo Const_EmitterY(TEXT("Emitter Y Axis"), ENiagaraDataType::Vector);
 		static FNiagaraVariableInfo Const_EmitterZ(TEXT("Emitter Z Axis"), ENiagaraDataType::Vector);
@@ -300,6 +301,7 @@ const TArray<FNiagaraVariableInfo>& UNiagaraComponent::GetSystemConstants()
 		static FNiagaraVariableInfo Const_Zero(TEXT("ZERO"), ENiagaraDataType::Vector);
 		static FNiagaraVariableInfo Const_DeltaTime(TEXT("Delta Time"), ENiagaraDataType::Vector);
 		static FNiagaraVariableInfo Const_EmitterPos(TEXT("Emitter Position"), ENiagaraDataType::Vector);
+		static FNiagaraVariableInfo Const_EmitterAge(TEXT("Emitter Age"), ENiagaraDataType::Vector);
 		static FNiagaraVariableInfo Const_EmitterX(TEXT("Emitter X Axis"), ENiagaraDataType::Vector);
 		static FNiagaraVariableInfo Const_EmitterY(TEXT("Emitter Y Axis"), ENiagaraDataType::Vector);
 		static FNiagaraVariableInfo Const_EmitterZ(TEXT("Emitter Z Axis"), ENiagaraDataType::Vector);
@@ -307,7 +309,7 @@ const TArray<FNiagaraVariableInfo>& UNiagaraComponent::GetSystemConstants()
 		SystemConstants.Add(Const_Zero);
 		SystemConstants.Add(Const_DeltaTime);
 		SystemConstants.Add(Const_EmitterPos);
-		SystemConstants.Add(Const_EmitterPos);
+		SystemConstants.Add(Const_EmitterAge);
 		SystemConstants.Add(Const_EmitterX);
 		SystemConstants.Add(Const_EmitterY);
 		SystemConstants.Add(Const_EmitterZ);
diff --git a/Engine/Source/Runtime/Engine/Private/NiagaraEffectRenderer.cpp b/Engine/Source/Runtime/Engine/Private/NiagaraEffectRenderer.cpp
--- a/Engine/Source/Runtime/Engine/Private/NiagaraEffectRenderer.cpp
+++ b/Engine/Source/Runtime/Engine/Private/NiagaraEffectRenderer.cpp
@@ -215,12 +215,13 @@ FNiagaraDynamicDataBase *NiagaraEffectRendererSprites::GenerateVertexData(const
 	SCOPE_CYCLE_COUNTER(STAT_NiagaraGenSpriteVertexData);
 
 	SimpleTimer VertexDataTimer;
+	FVector MaxSize;
 
 	FNiagaraDynamicDataSprites *DynamicData = new FNiagaraDynamicDataSprites;
 	TArray<FParticleSpriteVertex>& RenderData = DynamicData->VertexData;
 
 	RenderData.Reset(Data.GetNumParticles());
-	//CachedBounds.Init();
+	CachedBounds.Init();
 
 	const FVector4 *PosPtr = Data.GetAttributeData(FNiagaraVariableInfo(FName(TEXT("Position")), ENiagaraDataType::Vector));
 	const FVector4 *ColPtr = Data.GetAttributeData(FNiagaraVariableInfo(FName(TEXT("Color")), ENiagaraDataType::Vector));
@@ -257,10 +258,10 @@ FNiagaraDynamicDataBase *NiagaraEffectRendererSprites::GenerateVertexData(const
 		FPlatformMisc::Prefetch(RotPtr + ParticleIndex + 1);
 		FPlatformMisc::Prefetch(ColPtr + ParticleIndex + 1);		
 		FPlatformMisc::Prefetch(AgePtr + ParticleIndex + 1);
-
-		//CachedBounds += NewVertex.Position;
+		MaxSize = MaxSize.ComponentMax(FVector(NewVertex.Size.GetMax()));
+		CachedBounds += NewVertex.Position;
 	}
-	//CachedBounds.ExpandBy(MaxSize);
+	CachedBounds.ExpandBy(MaxSize);
 
 	CPUTimeMS = VertexDataTimer.GetElapsedMilliseconds();
 
diff --git a/Engine/Source/Runtime/VectorVM/Private/VectorVM.cpp b/Engine/Source/Runtime/VectorVM/Private/VectorVM.cpp
--- a/Engine/Source/Runtime/VectorVM/Private/VectorVM.cpp
+++ b/Engine/Source/Runtime/VectorVM/Private/VectorVM.cpp
@@ -2,7 +2,7 @@
 
 #include "VectorVMPrivate.h"
 #include "ModuleManager.h"
-
+#include <omp.h>
 
 IMPLEMENT_MODULE(FDefaultModuleImpl, VectorVM);
 
@@ -101,6 +101,13 @@ static /*VM_FORCEINLINE*/ const FNiagaraDataObject *DecodeDataObjConstant(FVecto
 	return Obj;
 }
 
+/** Decode a constant from the bytecode. */
+static /*VM_FORCEINLINE*/ FNiagaraDataObject *DecodeWritableDataObjConstant(FVectorVMContext& Context)
+{
+	FNiagaraDataObject* Obj = Context.DataObjConstantTable[*Context.Code++];
+	return Obj;
+}
+
 
 static VM_FORCEINLINE uint8 DecodeSrcOperandTypes(FVectorVMContext& Context)
 {
@@ -134,6 +141,17 @@ struct FDataObjectConstantHandler
 	VM_FORCEINLINE const FNiagaraDataObject *Get(){ return Constant; }
 };
 
+/** Handles reading of a data object constant. */
+struct FWritableDataObjectConstantHandler
+{
+	FNiagaraDataObject *Constant;
+	FWritableDataObjectConstantHandler(FVectorVMContext& Context)
+		: Constant(DecodeWritableDataObjConstant(Context))
+	{}
+	VM_FORCEINLINE FNiagaraDataObject *Get(){ return Constant; }
+};
+
+
 /** Handles reading of a register, advancing the pointer with each read. */
 struct FRegisterHandler
 {
@@ -252,6 +270,25 @@ struct TBinaryVectorKernelData
 
 
 
+/** Base class of Vector kernels with 2 operands, one of which can be a data object. */
+template <typename Kernel>
+struct TTrinaryVectorKernelData
+{
+	static void Exec(FVectorVMContext& Context)
+	{
+		VectorRegister* RESTRICT Dst = DecodeRegister(Context);
+		uint32 SrcOpTypes = DecodeSrcOperandTypes(Context);
+		int32 NumVectors = Context.NumVectors;
+		switch (SrcOpTypes)
+		{
+		case SRCOP_RRRB:	VectorTrinaryLoop<Kernel, FWritableDataObjectConstantHandler, FRegisterHandler, FRegisterHandler>(Context, Dst, NumVectors); break;
+		default: check(0); break;
+		};
+	}
+};
+
+
+
 
 /** Base class of Vector kernels with 3 operands. */
 template <typename Kernel>
@@ -377,12 +414,7 @@ struct FVectorKernelSqrt : public TUnaryVectorKernel<FVectorKernelSqrt>
 	static void VM_FORCEINLINE DoKernel(VectorRegister* RESTRICT Dst,VectorRegister Src0)
 	{
 		// TODO: Need a SIMD sqrt!
-		float* RESTRICT FloatDst = reinterpret_cast<float* RESTRICT>(Dst);
-		float const* FloatSrc0 = reinterpret_cast<float const*>(&Src0);
-		FloatDst[0] = FMath::Sqrt(FloatSrc0[0]);
-		FloatDst[1] = FMath::Sqrt(FloatSrc0[1]);
-		FloatDst[2] = FMath::Sqrt(FloatSrc0[2]);
-		FloatDst[3] = FMath::Sqrt(FloatSrc0[3]);
+		*Dst = VectorReciprocal(VectorReciprocalSqrt(Src0));
 	}
 };
 
@@ -562,7 +594,21 @@ struct FVectorKernelSample : public TBinaryVectorKernelData<FVectorKernelSample>
 		if (Src0)
 		{
 			float const* FloatSrc1 = reinterpret_cast<float const*>(&Src1);
-			FVector4 Tmp = Src0->Sample(FloatSrc1[0]);
+			FVector4 Tmp = Src0->Sample(FVector4(FloatSrc1[0], FloatSrc1[1], FloatSrc1[2], FloatSrc1[3]));
+			*Dst = VectorLoad(&Tmp);
+		}
+	}
+};
+
+struct FVectorKernelBufferWrite : public TTrinaryVectorKernelData<FVectorKernelBufferWrite>
+{
+	static void FORCEINLINE DoKernel(VectorRegister* RESTRICT Dst, FNiagaraDataObject *Src0, VectorRegister Src1, VectorRegister Src2)
+	{
+		if (Src0)
+		{
+			float const* FloatSrc1 = reinterpret_cast<float const*>(&Src1);	// Coords
+			float const* FloatSrc2 = reinterpret_cast<float const*>(&Src2);	// Value
+			FVector4 Tmp = Src0->Write(FVector4(FloatSrc1[0], FloatSrc1[1], FloatSrc1[2], FloatSrc1[3]), FVector4(FloatSrc2[0], FloatSrc2[1], FloatSrc2[2], FloatSrc2[3]));
 			*Dst = VectorLoad(&Tmp);
 		}
 	}
@@ -855,6 +901,7 @@ void VectorVM::Exec(
 			case EOp::composew: FVectorKernelCompose<3, 3, 3, 3>::Exec(Context); break;
 			case EOp::lessthan: FVectorKernelLessThan::Exec(Context); break;
 			case EOp::sample: FVectorKernelSample::Exec(Context); break;
+			case EOp::bufferwrite: FVectorKernelBufferWrite::Exec(Context); break;
 			case EOp::output: FVectorKernelOutput::Exec(Context); break;
 
 			// Execution always terminates with a "done" opcode.
diff --git a/Engine/Source/Runtime/VectorVM/Public/VectorVM.h b/Engine/Source/Runtime/VectorVM/Public/VectorVM.h
--- a/Engine/Source/Runtime/VectorVM/Public/VectorVM.h
+++ b/Engine/Source/Runtime/VectorVM/Public/VectorVM.h
@@ -11,7 +11,8 @@
 class FNiagaraDataObject
 {
 public:
-	virtual FVector4 Sample(float InTime) const = 0;
+	virtual FVector4 Sample(FVector4 InCoords) const = 0;
+	virtual FVector4 Write(FVector4 InCoords, FVector4 Value) = 0;
 };
 
 /* Curve object; encapsulates a curve for the VectorVM
@@ -25,7 +26,7 @@ public:
 	{
 	}
 
-	virtual FVector4 Sample(float InTime) const override
+	virtual FVector4 Sample(FVector4 InCoords) const override
 	{
 		// commented out for the time being; until Niagara is in its own module, we can't depend on curve objects here
 		// because they're part of engine
@@ -34,11 +35,56 @@ public:
 		return FVector4(1.0f, 0.0f, 0.0f, 1.0f);
 	}
 
+	virtual FVector4 Write(FVector4 InCoords, FVector4 Value) override
+	{
+		return FVector4(1.0f, 0.0f, 0.0f, 1.0f);
+	}
+
 	UCurveVector *GetCurveObject()	{ return CurveObj; }
 	void SetCurveObject(UCurveVector *InCurve)	{ CurveObj = InCurve; }
 };
 
 
+/* Curve object; encapsulates sparse volumetric data for the VectorVM
+*/
+class FNiagaraSparseVolumeDataObject : public FNiagaraDataObject
+{
+private:
+	TArray<FVector4> Data;
+public:
+	FNiagaraSparseVolumeDataObject()
+	{
+		Data.AddZeroed(32768);
+	}
+
+	uint32 MakeHash(FVector4 Coords) const
+	{
+		const uint32 P1 = 73856093; // some large primes
+		const uint32 P2 = 19349663;
+		const uint32 P3 = 83492791;
+		Coords += FVector4(500.0f, 500.0f, 500.0f, 500.0f);
+		Coords *= 0.02f;
+		Coords = FVector4(FMath::Max(Coords.X, 0.0f), FMath::Max(Coords.Y, 0.0f), FMath::Max(Coords.Z, 0.0f), FMath::Max(Coords.W, 0.0f));
+		int N = (P1*static_cast<uint32>(Coords.X)) ^ (P2*static_cast<uint32>(Coords.Y)) ^ (P3*static_cast<uint32>(Coords.Z));
+		N %= 32768;
+		return FMath::Clamp(N, 0, 32768);
+	}
+
+	virtual FVector4 Sample(FVector4 Coords) const override
+	{
+		int32 Index = MakeHash(Coords);
+		Index = FMath::Clamp(Index, 0, Data.Num() - 1);
+		return Data[Index];
+	}
+
+	virtual FVector4 Write(FVector4 InCoords, FVector4 Value) override
+	{
+		Data[MakeHash(InCoords)] = Value;
+		return Value;
+	}
+};
+
+
 namespace VectorVM
 {
 	/** Constants. */
@@ -117,6 +163,8 @@ namespace VectorVM
 			output,
 			lessthan,
 			sample,
+			bufferwrite,
+			eventbroadcast,
 			NumOpcodes
 	};
 
