Various UI improvements for the graph editor.

- Pins now auto wire to new nodes created when draging from a pin.
- Better error reporting and ability to fail compilation. Still WIP.
- Added new pin context menu for breaking existing links.
- Better Undo/Redo behavior.
- Title of input nodes is now "Input" rather than the attirubte/constant they retrieived.

diff --git a/Engine/Source/Editor/NiagaraEditor/Private/NiagaraCompiler.cpp b/Engine/Source/Editor/NiagaraEditor/Private/NiagaraCompiler.cpp
--- a/Engine/Source/Editor/NiagaraEditor/Private/NiagaraCompiler.cpp
+++ b/Engine/Source/Editor/NiagaraEditor/Private/NiagaraCompiler.cpp
@@ -123,36 +123,64 @@ public:
 
 //////////////////////////////////////////////////////////////////////////
 
-void FNiagaraCompiler::CheckInputs(FName OpName, TArray<TNiagaraExprPtr>& Inputs)
+bool FNiagaraCompiler::CheckInputs(FName OpName, TArray<TNiagaraExprPtr>& Inputs)
 {
 	//check the types of the input expressions.
 	const FNiagaraOpInfo* OpInfo = FNiagaraOpInfo::GetOpInfo(OpName);
 	check(OpInfo);
 	int32 NumInputs = Inputs.Num();
 	check(OpInfo->Inputs.Num() == NumInputs);
+	bool bError = false;
 	for (int32 i = 0; i < NumInputs ; ++i)
 	{
+		check(Inputs[i].IsValid());
 		if (OpInfo->Inputs[i].DataType != Inputs[i]->Result.Type)
 		{
-			UE_LOG(LogNiagaraCompiler, Error, TEXT("Exression %s has incorrect inputs!\nExpected: %d - Actual: %d"), *OpName.ToString(), (int32)OpInfo->Inputs[i].DataType, (int32)((TNiagaraExprPtr)Inputs[i])->Result.Type.GetValue());
+			bError = true;
+			FText ErrorText = FText::Format(LOCTEXT("Expression {0} has incorrect inputs!\nExpected: {1} - Actual: {2}", ""),
+				FText::FromString(OpName.ToString()),
+				FText::AsNumber((int32)OpInfo->Inputs[i].DataType),
+				FText::AsNumber((int32)((TNiagaraExprPtr)Inputs[i])->Result.Type.GetValue()));
+			MessageLog.Error(*ErrorText.ToString());
 		}
 	}
+	return bError;
 }
 
-void FNiagaraCompiler::CheckOutputs(FName OpName, TArray<TNiagaraExprPtr>& Outputs)
+bool FNiagaraCompiler::CheckOutputs(FName OpName, TArray<TNiagaraExprPtr>& Outputs)
 {
 	//check the types of the input expressions.
 	const FNiagaraOpInfo* OpInfo = FNiagaraOpInfo::GetOpInfo(OpName);
 	check(OpInfo);
 	int32 NumOutputs = Outputs.Num();
+	bool bError = false;
 	check(OpInfo->Outputs.Num() == NumOutputs);
 	for (int32 i = 0; i < NumOutputs; ++i)
 	{
+		check(Outputs[i].IsValid());
 		if (OpInfo->Outputs[i].DataType != Outputs[i]->Result.Type)
 		{
-			UE_LOG(LogNiagaraCompiler, Error, TEXT("Exression %s has incorrect outputs!\nExpected: %d - Actual: %d"), *OpName.ToString(), (int32)OpInfo->Outputs[i].DataType, (int32)((TNiagaraExprPtr)Outputs[i])->Result.Type.GetValue());
+			bError = true;
+			FText ErrorText = FText::Format(LOCTEXT("Expression {0} has incorrect inputs!\nExpected: {1} - Actual: {2}", ""),
+				FText::FromString(OpName.ToString()),
+				FText::AsNumber((int32)OpInfo->Outputs[i].DataType),
+				FText::AsNumber((int32)((TNiagaraExprPtr)Outputs[i])->Result.Type.GetValue()));
+			MessageLog.Error(*ErrorText.ToString());
 		}
 	}
+	return bError;
+}
+
+void FNiagaraCompiler::Error(FText ErrorText, UNiagaraNode* Node, UEdGraphPin* Pin)
+{
+	FString ErrorString = FString::Printf(TEXT("Node: @@ - Pin: @@ - %s"), *ErrorText.ToString());
+	MessageLog.Error(*ErrorString, Node, Pin);
+}
+
+void FNiagaraCompiler::Warning(FText WarningText, UNiagaraNode* Node, UEdGraphPin* Pin)
+{
+	FString WarnString = FString::Printf(TEXT("Node: @@ - Pin: @@ - %s"), *WarningText.ToString());
+	MessageLog.Warning(*WarnString, Node, Pin);
 }
 
 int32 FNiagaraCompiler::GetAttributeIndex(const FNiagaraVariableInfo& Attr)const
diff --git a/Engine/Source/Editor/NiagaraEditor/Private/NiagaraCompiler.h b/Engine/Source/Editor/NiagaraEditor/Private/NiagaraCompiler.h
--- a/Engine/Source/Editor/NiagaraEditor/Private/NiagaraCompiler.h
+++ b/Engine/Source/Editor/NiagaraEditor/Private/NiagaraCompiler.h
@@ -67,8 +67,11 @@ public:
 	virtual TNiagaraExprPtr GetInternalConstant(const FNiagaraVariableInfo& Constant, const FMatrix& Default)override;
 	virtual TNiagaraExprPtr GetExternalCurveConstant(const FNiagaraVariableInfo& Constant)override;
 
-	virtual void CheckInputs(FName OpName, TArray<TNiagaraExprPtr>& Inputs)override;
-	virtual void CheckOutputs(FName OpName, TArray<TNiagaraExprPtr>& Outputs)override;
+	virtual bool CheckInputs(FName OpName, TArray<TNiagaraExprPtr>& Inputs)override;
+	virtual bool CheckOutputs(FName OpName, TArray<TNiagaraExprPtr>& Outputs)override;
+
+	virtual void Error(FText ErrorText, UNiagaraNode* Node, UEdGraphPin* Pin)override;
+	virtual void Warning(FText WarningText, UNiagaraNode* Node, UEdGraphPin* Pin)override;
 	//End INiagaraCompiler Interface
 	
 	/** Gets the index into a constants table of the constant specified by Name and bInternal. */
diff --git a/Engine/Source/Editor/NiagaraEditor/Private/NiagaraCompiler_VectorVM.cpp b/Engine/Source/Editor/NiagaraEditor/Private/NiagaraCompiler_VectorVM.cpp
--- a/Engine/Source/Editor/NiagaraEditor/Private/NiagaraCompiler_VectorVM.cpp
+++ b/Engine/Source/Editor/NiagaraEditor/Private/NiagaraCompiler_VectorVM.cpp
@@ -189,7 +189,7 @@ TNiagaraExprPtr FNiagaraCompiler_VectorVM::Output(const FNiagaraVariableInfo& At
 	return Expressions[Index];
 }
 
-void FNiagaraCompiler_VectorVM::CompileScript(UNiagaraScript* InScript)
+bool FNiagaraCompiler_VectorVM::CompileScript(UNiagaraScript* InScript)
 {
 	//TODO - Most of this function can be refactored up into FNiagaraCompiler once we know how the compute script and VM script will coexist in the UNiagaraScript.
 
@@ -198,6 +198,27 @@ void FNiagaraCompiler_VectorVM::CompileScript(UNiagaraScript* InScript)
 	Script = InScript;
 	Source = CastChecked<UNiagaraScriptSource>(InScript->Source);
 
+	//Clear previous graph errors.
+	if (InScript->ByteCode.Num() == 0)
+	{
+		bool bChanged = false;
+		for (UEdGraphNode* Node : Source->NodeGraph->Nodes)
+		{
+			if (Node->bHasCompilerMessage)
+			{
+				Node->ErrorMsg.Empty();
+				Node->ErrorType = EMessageSeverity::Info;
+				Node->bHasCompilerMessage = false;
+				Node->Modify(true);
+				bChanged = true;
+			}
+		}
+		if (bChanged)
+		{
+			Source->NodeGraph->NotifyGraphChanged();
+		}
+	}
+	
 	//Should we roll our own message/error log and put it in a window somewhere?
 	MessageLog.SetSourceName(InScript->GetPathName());
 
@@ -205,51 +226,57 @@ void FNiagaraCompiler_VectorVM::CompileScript(UNiagaraScript* InScript)
 	SourceGraph = CastChecked<UNiagaraGraph>(FEdGraphUtilities::CloneGraph(Source->NodeGraph, NULL, &MessageLog));
 	FEdGraphUtilities::MergeChildrenGraphsIn(SourceGraph, SourceGraph, /*bRequireSchemaMatch=*/ true);
 
-	if (!MergeInFunctionNodes())
-	{ 
-		//TODO: Insert a good error reporting and compile failure system.
-		InScript->Attributes.Empty();
-		InScript->ByteCode.Empty();
-		InScript->ConstantData.Empty();
-		return;
-	}
-
-	TempRegisters.AddZeroed(VectorVM::NumTempRegisters);
+	if (MergeInFunctionNodes())
+	{
+		TempRegisters.AddZeroed(VectorVM::NumTempRegisters);
 
-	// Find the output node.
-	UNiagaraNodeOutput* OutputNode = SourceGraph->FindOutputNode();
+		// Find the output node.
+		UNiagaraNodeOutput* OutputNode = SourceGraph->FindOutputNode();
 
-	TArray<FNiagaraNodeResult> OutputExpressions;
-	OutputNode->Compile(this, OutputExpressions);
+		TArray<FNiagaraNodeResult> OutputExpressions;
+		OutputNode->Compile(this, OutputExpressions);
 
-	//Todo - track usage of temp registers so we can allocate and free and reuse from a smaller set.
-	int32 NextTempRegister = 0;
-	Script->ByteCode.Empty();
+		//Todo - track usage of temp registers so we can allocate and free and reuse from a smaller set.
+		int32 NextTempRegister = 0;
+		Script->ByteCode.Empty();
 
-	//Now we have a set of expressions from which we can generate the bytecode.
-	for (int32 i = 0; i < Expressions.Num() ; ++i)
-	{
-		TNiagaraExprPtr Expr = Expressions[i];
-		if (Expr.IsValid())
+		//Now we have a set of expressions from which we can generate the bytecode.
+		for (int32 i = 0; i < Expressions.Num(); ++i)
 		{
-			Expr->Process();
-			Expr->PostProcess();
+			TNiagaraExprPtr Expr = Expressions[i];
+			if (Expr.IsValid())
+			{
+				Expr->Process();
+				Expr->PostProcess();
+			}
+			else
+			{
+				check(0);//Replace with graceful bailout?
+			}
+
+			Expressions[i].Reset();
 		}
-		else
+
+		PinToExpression.Empty();//TODO - Replace this with expression caching at a lower level. Will be more useful and nicer.
+
+		if (MessageLog.NumErrors == 0)
 		{
-			check(0);//Replace with graceful bailout?
+			// Terminate with the 'done' opcode.
+			Script->ByteCode.Add((int8)VectorVM::EOp::done);
+			Script->ConstantData = ConstantData;
+			Script->Attributes = OutputNode->Outputs;
+			return true;
 		}
-
-		Expressions[i].Reset();
 	}
 
-	PinToExpression.Empty();//TODO - Replace this with expression caching at a lower level. Will be more useful and nicer.
-
-	// Terminate with the 'done' opcode.
-	Script->ByteCode.Add((int8)VectorVM::EOp::done);
-
-	Script->ConstantData = ConstantData;
-	Script->Attributes = OutputNode->Outputs;
+	//Compilation has failed.
+	InScript->Attributes.Empty();
+	InScript->ByteCode.Empty();
+	InScript->ConstantData.Empty();
+	//Tell the graph to show the errors.
+	Source->NodeGraph->NotifyGraphChanged();
+		
+	return false;
 }
 
 TNiagaraExprPtr FNiagaraCompiler_VectorVM::Expression_VMNative(VectorVM::EOp Op, TArray<TNiagaraExprPtr>& InputExpressions)
@@ -283,273 +310,325 @@ TNiagaraExprPtr FNiagaraCompiler_VectorVM::Expression_VMNative(VectorVM::EOp Op,
 	return Expression_VMNative(Op, In);
 }
 
-void FNiagaraCompiler_VectorVM::Add_Internal(TArray<TNiagaraExprPtr>& InputExpressions, TArray<TNiagaraExprPtr>& OutputExpressions)
+bool FNiagaraCompiler_VectorVM::Add_Internal(TArray<TNiagaraExprPtr>& InputExpressions, TArray<TNiagaraExprPtr>& OutputExpressions)
 {
 	OutputExpressions.Add(Expression_VMNative(VectorVM::EOp::add, InputExpressions));
+	return true;
 }
 
-void FNiagaraCompiler_VectorVM::Subtract_Internal(TArray<TNiagaraExprPtr>& InputExpressions, TArray<TNiagaraExprPtr>& OutputExpressions)
+bool FNiagaraCompiler_VectorVM::Subtract_Internal(TArray<TNiagaraExprPtr>& InputExpressions, TArray<TNiagaraExprPtr>& OutputExpressions)
 {
 	OutputExpressions.Add((Expression_VMNative(VectorVM::EOp::sub, InputExpressions)));
+	return true;
 }
 
-void FNiagaraCompiler_VectorVM::Multiply_Internal(TArray<TNiagaraExprPtr>& InputExpressions, TArray<TNiagaraExprPtr>& OutputExpressions)
+bool FNiagaraCompiler_VectorVM::Multiply_Internal(TArray<TNiagaraExprPtr>& InputExpressions, TArray<TNiagaraExprPtr>& OutputExpressions)
 {
 	OutputExpressions.Add(Expression_VMNative(VectorVM::EOp::mul, InputExpressions));
+	return true;
 }
 
-void FNiagaraCompiler_VectorVM::Divide_Internal(TArray<TNiagaraExprPtr>& InputExpressions, TArray<TNiagaraExprPtr>& OutputExpressions)
+bool FNiagaraCompiler_VectorVM::Divide_Internal(TArray<TNiagaraExprPtr>& InputExpressions, TArray<TNiagaraExprPtr>& OutputExpressions)
 {
 	OutputExpressions.Add(Expression_VMNative(VectorVM::EOp::div, InputExpressions));
+	return true;
 }
 
-void FNiagaraCompiler_VectorVM::MultiplyAdd_Internal(TArray<TNiagaraExprPtr>& InputExpressions, TArray<TNiagaraExprPtr>& OutputExpressions)
+bool FNiagaraCompiler_VectorVM::MultiplyAdd_Internal(TArray<TNiagaraExprPtr>& InputExpressions, TArray<TNiagaraExprPtr>& OutputExpressions)
 {
 	OutputExpressions.Add(Expression_VMNative(VectorVM::EOp::mad, InputExpressions));
+	return true;
 }
 
-void FNiagaraCompiler_VectorVM::Lerp_Internal(TArray<TNiagaraExprPtr>& InputExpressions, TArray<TNiagaraExprPtr>& OutputExpressions)
+bool FNiagaraCompiler_VectorVM::Lerp_Internal(TArray<TNiagaraExprPtr>& InputExpressions, TArray<TNiagaraExprPtr>& OutputExpressions)
 {
 	OutputExpressions.Add(Expression_VMNative(VectorVM::EOp::lerp, InputExpressions));
+	return true;
 }
 
-void FNiagaraCompiler_VectorVM::Reciprocal_Internal(TArray<TNiagaraExprPtr>& InputExpressions, TArray<TNiagaraExprPtr>& OutputExpressions)
+bool FNiagaraCompiler_VectorVM::Reciprocal_Internal(TArray<TNiagaraExprPtr>& InputExpressions, TArray<TNiagaraExprPtr>& OutputExpressions)
 {
 	OutputExpressions.Add(Expression_VMNative(VectorVM::EOp::rcp, InputExpressions));
+	return true;
 }
 
-void FNiagaraCompiler_VectorVM::ReciprocalSqrt_Internal(TArray<TNiagaraExprPtr>& InputExpressions, TArray<TNiagaraExprPtr>& OutputExpressions)
+bool FNiagaraCompiler_VectorVM::ReciprocalSqrt_Internal(TArray<TNiagaraExprPtr>& InputExpressions, TArray<TNiagaraExprPtr>& OutputExpressions)
 {
 	OutputExpressions.Add(Expression_VMNative(VectorVM::EOp::rsq, InputExpressions));
+	return true;
 }
 
-void FNiagaraCompiler_VectorVM::Sqrt_Internal(TArray<TNiagaraExprPtr>& InputExpressions, TArray<TNiagaraExprPtr>& OutputExpressions)
+bool FNiagaraCompiler_VectorVM::Sqrt_Internal(TArray<TNiagaraExprPtr>& InputExpressions, TArray<TNiagaraExprPtr>& OutputExpressions)
 {
 	OutputExpressions.Add(Expression_VMNative(VectorVM::EOp::sqrt, InputExpressions));
+	return true;
 }
 
-void FNiagaraCompiler_VectorVM::Negate_Internal(TArray<TNiagaraExprPtr>& InputExpressions, TArray<TNiagaraExprPtr>& OutputExpressions)
+bool FNiagaraCompiler_VectorVM::Negate_Internal(TArray<TNiagaraExprPtr>& InputExpressions, TArray<TNiagaraExprPtr>& OutputExpressions)
 {
 	OutputExpressions.Add(Expression_VMNative(VectorVM::EOp::neg, InputExpressions));
+	return true;
 }
 
-void FNiagaraCompiler_VectorVM::Abs_Internal(TArray<TNiagaraExprPtr>& InputExpressions, TArray<TNiagaraExprPtr>& OutputExpressions)
+bool FNiagaraCompiler_VectorVM::Abs_Internal(TArray<TNiagaraExprPtr>& InputExpressions, TArray<TNiagaraExprPtr>& OutputExpressions)
 {
 	OutputExpressions.Add(Expression_VMNative(VectorVM::EOp::abs, InputExpressions));
+	return true;
 }
 
-void FNiagaraCompiler_VectorVM::Exp_Internal(TArray<TNiagaraExprPtr>& InputExpressions, TArray<TNiagaraExprPtr>& OutputExpressions)
+bool FNiagaraCompiler_VectorVM::Exp_Internal(TArray<TNiagaraExprPtr>& InputExpressions, TArray<TNiagaraExprPtr>& OutputExpressions)
 {
 	OutputExpressions.Add(Expression_VMNative(VectorVM::EOp::exp, InputExpressions));
+	return true;
 }
 
-void FNiagaraCompiler_VectorVM::Exp2_Internal(TArray<TNiagaraExprPtr>& InputExpressions, TArray<TNiagaraExprPtr>& OutputExpressions)
+bool FNiagaraCompiler_VectorVM::Exp2_Internal(TArray<TNiagaraExprPtr>& InputExpressions, TArray<TNiagaraExprPtr>& OutputExpressions)
 {
 	OutputExpressions.Add(Expression_VMNative(VectorVM::EOp::exp2, InputExpressions));
+	return true;
 }
 
-void FNiagaraCompiler_VectorVM::Log_Internal(TArray<TNiagaraExprPtr>& InputExpressions, TArray<TNiagaraExprPtr>& OutputExpressions)
+bool FNiagaraCompiler_VectorVM::Log_Internal(TArray<TNiagaraExprPtr>& InputExpressions, TArray<TNiagaraExprPtr>& OutputExpressions)
 {
 	OutputExpressions.Add(Expression_VMNative(VectorVM::EOp::log, InputExpressions));
+	return true;
 }
 
-void FNiagaraCompiler_VectorVM::LogBase2_Internal(TArray<TNiagaraExprPtr>& InputExpressions, TArray<TNiagaraExprPtr>& OutputExpressions)
+bool FNiagaraCompiler_VectorVM::LogBase2_Internal(TArray<TNiagaraExprPtr>& InputExpressions, TArray<TNiagaraExprPtr>& OutputExpressions)
 {
 	OutputExpressions.Add(Expression_VMNative(VectorVM::EOp::log2, InputExpressions));
+	return true;
 }
 
-void FNiagaraCompiler_VectorVM::Sin_Internal(TArray<TNiagaraExprPtr>& InputExpressions, TArray<TNiagaraExprPtr>& OutputExpressions)
+bool FNiagaraCompiler_VectorVM::Sin_Internal(TArray<TNiagaraExprPtr>& InputExpressions, TArray<TNiagaraExprPtr>& OutputExpressions)
 {
 	OutputExpressions.Add(Expression_VMNative(VectorVM::EOp::sin, InputExpressions));
+	return true;
 }
 
-void FNiagaraCompiler_VectorVM::Cos_Internal(TArray<TNiagaraExprPtr>& InputExpressions, TArray<TNiagaraExprPtr>& OutputExpressions)
+bool FNiagaraCompiler_VectorVM::Cos_Internal(TArray<TNiagaraExprPtr>& InputExpressions, TArray<TNiagaraExprPtr>& OutputExpressions)
 {
 	OutputExpressions.Add(Expression_VMNative(VectorVM::EOp::cos, InputExpressions));
+	return true;
 }
 
-void FNiagaraCompiler_VectorVM::Tan_Internal(TArray<TNiagaraExprPtr>& InputExpressions, TArray<TNiagaraExprPtr>& OutputExpressions)
+bool FNiagaraCompiler_VectorVM::Tan_Internal(TArray<TNiagaraExprPtr>& InputExpressions, TArray<TNiagaraExprPtr>& OutputExpressions)
 {
 	OutputExpressions.Add(Expression_VMNative(VectorVM::EOp::tan, InputExpressions));
+	return true;
 }
 
-void FNiagaraCompiler_VectorVM::ASin_Internal(TArray<TNiagaraExprPtr>& InputExpressions, TArray<TNiagaraExprPtr>& OutputExpressions)
+bool FNiagaraCompiler_VectorVM::ASin_Internal(TArray<TNiagaraExprPtr>& InputExpressions, TArray<TNiagaraExprPtr>& OutputExpressions)
 {
 	OutputExpressions.Add(Expression_VMNative(VectorVM::EOp::asin, InputExpressions));
+	return true;
 }
 
-void FNiagaraCompiler_VectorVM::ACos_Internal(TArray<TNiagaraExprPtr>& InputExpressions, TArray<TNiagaraExprPtr>& OutputExpressions)
+bool FNiagaraCompiler_VectorVM::ACos_Internal(TArray<TNiagaraExprPtr>& InputExpressions, TArray<TNiagaraExprPtr>& OutputExpressions)
 {
 	OutputExpressions.Add(Expression_VMNative(VectorVM::EOp::acos, InputExpressions));
+	return true;
 }
 
-void FNiagaraCompiler_VectorVM::ATan_Internal(TArray<TNiagaraExprPtr>& InputExpressions, TArray<TNiagaraExprPtr>& OutputExpressions)
+bool FNiagaraCompiler_VectorVM::ATan_Internal(TArray<TNiagaraExprPtr>& InputExpressions, TArray<TNiagaraExprPtr>& OutputExpressions)
 {
 	OutputExpressions.Add(Expression_VMNative(VectorVM::EOp::atan, InputExpressions));
+	return true;
 }
 
-void FNiagaraCompiler_VectorVM::ATan2_Internal(TArray<TNiagaraExprPtr>& InputExpressions, TArray<TNiagaraExprPtr>& OutputExpressions)
+bool FNiagaraCompiler_VectorVM::ATan2_Internal(TArray<TNiagaraExprPtr>& InputExpressions, TArray<TNiagaraExprPtr>& OutputExpressions)
 {
 	OutputExpressions.Add(Expression_VMNative(VectorVM::EOp::atan2, InputExpressions));
+	return true;
 }
 
-void FNiagaraCompiler_VectorVM::Ceil_Internal(TArray<TNiagaraExprPtr>& InputExpressions, TArray<TNiagaraExprPtr>& OutputExpressions)
+bool FNiagaraCompiler_VectorVM::Ceil_Internal(TArray<TNiagaraExprPtr>& InputExpressions, TArray<TNiagaraExprPtr>& OutputExpressions)
 {
 	OutputExpressions.Add(Expression_VMNative(VectorVM::EOp::ceil, InputExpressions));
+	return true;
 }
 
-void FNiagaraCompiler_VectorVM::Floor_Internal(TArray<TNiagaraExprPtr>& InputExpressions, TArray<TNiagaraExprPtr>& OutputExpressions)
+bool FNiagaraCompiler_VectorVM::Floor_Internal(TArray<TNiagaraExprPtr>& InputExpressions, TArray<TNiagaraExprPtr>& OutputExpressions)
 {
 	OutputExpressions.Add(Expression_VMNative(VectorVM::EOp::floor, InputExpressions));
+	return true;
 }
 
-void FNiagaraCompiler_VectorVM::Modulo_Internal(TArray<TNiagaraExprPtr>& InputExpressions, TArray<TNiagaraExprPtr>& OutputExpressions)
+bool FNiagaraCompiler_VectorVM::Modulo_Internal(TArray<TNiagaraExprPtr>& InputExpressions, TArray<TNiagaraExprPtr>& OutputExpressions)
 {
 	OutputExpressions.Add(Expression_VMNative(VectorVM::EOp::fmod, InputExpressions));
+	return true;
 }
 
-void FNiagaraCompiler_VectorVM::Fractional_Internal(TArray<TNiagaraExprPtr>& InputExpressions, TArray<TNiagaraExprPtr>& OutputExpressions)
+bool FNiagaraCompiler_VectorVM::Fractional_Internal(TArray<TNiagaraExprPtr>& InputExpressions, TArray<TNiagaraExprPtr>& OutputExpressions)
 {
 	OutputExpressions.Add(Expression_VMNative(VectorVM::EOp::frac, InputExpressions));
+	return true;
 }
 
-void FNiagaraCompiler_VectorVM::Truncate_Internal(TArray<TNiagaraExprPtr>& InputExpressions, TArray<TNiagaraExprPtr>& OutputExpressions)
+bool FNiagaraCompiler_VectorVM::Truncate_Internal(TArray<TNiagaraExprPtr>& InputExpressions, TArray<TNiagaraExprPtr>& OutputExpressions)
 {
 	OutputExpressions.Add(Expression_VMNative(VectorVM::EOp::trunc, InputExpressions));
+	return true;
 }
 
-void FNiagaraCompiler_VectorVM::Clamp_Internal(TArray<TNiagaraExprPtr>& InputExpressions, TArray<TNiagaraExprPtr>& OutputExpressions)
+bool FNiagaraCompiler_VectorVM::Clamp_Internal(TArray<TNiagaraExprPtr>& InputExpressions, TArray<TNiagaraExprPtr>& OutputExpressions)
 {
 	OutputExpressions.Add(Expression_VMNative(VectorVM::EOp::clamp, InputExpressions));
+	return true;
 }
 
-void FNiagaraCompiler_VectorVM::Min_Internal(TArray<TNiagaraExprPtr>& InputExpressions, TArray<TNiagaraExprPtr>& OutputExpressions)
+bool FNiagaraCompiler_VectorVM::Min_Internal(TArray<TNiagaraExprPtr>& InputExpressions, TArray<TNiagaraExprPtr>& OutputExpressions)
 {
 	OutputExpressions.Add(Expression_VMNative(VectorVM::EOp::min, InputExpressions));
+	return true;
 }
 
-void FNiagaraCompiler_VectorVM::Max_Internal(TArray<TNiagaraExprPtr>& InputExpressions, TArray<TNiagaraExprPtr>& OutputExpressions)
+bool FNiagaraCompiler_VectorVM::Max_Internal(TArray<TNiagaraExprPtr>& InputExpressions, TArray<TNiagaraExprPtr>& OutputExpressions)
 {
 	OutputExpressions.Add(Expression_VMNative(VectorVM::EOp::max, InputExpressions));
+	return true;
 }
 
-void FNiagaraCompiler_VectorVM::Pow_Internal(TArray<TNiagaraExprPtr>& InputExpressions, TArray<TNiagaraExprPtr>& OutputExpressions)
+bool FNiagaraCompiler_VectorVM::Pow_Internal(TArray<TNiagaraExprPtr>& InputExpressions, TArray<TNiagaraExprPtr>& OutputExpressions)
 {
 	OutputExpressions.Add(Expression_VMNative(VectorVM::EOp::pow, InputExpressions));
+	return true;
 }
 
-void FNiagaraCompiler_VectorVM::Sign_Internal(TArray<TNiagaraExprPtr>& InputExpressions, TArray<TNiagaraExprPtr>& OutputExpressions)
+bool FNiagaraCompiler_VectorVM::Sign_Internal(TArray<TNiagaraExprPtr>& InputExpressions, TArray<TNiagaraExprPtr>& OutputExpressions)
 {
 	OutputExpressions.Add(Expression_VMNative(VectorVM::EOp::sign, InputExpressions));
+	return true;
 }
 
-void FNiagaraCompiler_VectorVM::Step_Internal(TArray<TNiagaraExprPtr>& InputExpressions, TArray<TNiagaraExprPtr>& OutputExpressions)
+bool FNiagaraCompiler_VectorVM::Step_Internal(TArray<TNiagaraExprPtr>& InputExpressions, TArray<TNiagaraExprPtr>& OutputExpressions)
 {
 	OutputExpressions.Add(Expression_VMNative(VectorVM::EOp::step, InputExpressions));
+	return true;
 }
 
-void FNiagaraCompiler_VectorVM::Dot_Internal(TArray<TNiagaraExprPtr>& InputExpressions, TArray<TNiagaraExprPtr>& OutputExpressions)
+bool FNiagaraCompiler_VectorVM::Dot_Internal(TArray<TNiagaraExprPtr>& InputExpressions, TArray<TNiagaraExprPtr>& OutputExpressions)
 {
 	OutputExpressions.Add(Expression_VMNative(VectorVM::EOp::dot, InputExpressions));
+	return true;
 }
 
-void FNiagaraCompiler_VectorVM::Cross_Internal(TArray<TNiagaraExprPtr>& InputExpressions, TArray<TNiagaraExprPtr>& OutputExpressions)
+bool FNiagaraCompiler_VectorVM::Cross_Internal(TArray<TNiagaraExprPtr>& InputExpressions, TArray<TNiagaraExprPtr>& OutputExpressions)
 {
 	OutputExpressions.Add(Expression_VMNative(VectorVM::EOp::cross, InputExpressions));
+	return true;
 }
 
-void FNiagaraCompiler_VectorVM::Normalize_Internal(TArray<TNiagaraExprPtr>& InputExpressions, TArray<TNiagaraExprPtr>& OutputExpressions)
+bool FNiagaraCompiler_VectorVM::Normalize_Internal(TArray<TNiagaraExprPtr>& InputExpressions, TArray<TNiagaraExprPtr>& OutputExpressions)
 {
 	OutputExpressions.Add(Expression_VMNative(VectorVM::EOp::normalize, InputExpressions));
+	return true;
 }
 
-void FNiagaraCompiler_VectorVM::Random_Internal(TArray<TNiagaraExprPtr>& InputExpressions, TArray<TNiagaraExprPtr>& OutputExpressions)
+bool FNiagaraCompiler_VectorVM::Random_Internal(TArray<TNiagaraExprPtr>& InputExpressions, TArray<TNiagaraExprPtr>& OutputExpressions)
 {
 	OutputExpressions.Add(Expression_VMNative(VectorVM::EOp::random, InputExpressions));
+	return true;
 }
 
-void FNiagaraCompiler_VectorVM::Length_Internal(TArray<TNiagaraExprPtr>& InputExpressions, TArray<TNiagaraExprPtr>& OutputExpressions)
+bool FNiagaraCompiler_VectorVM::Length_Internal(TArray<TNiagaraExprPtr>& InputExpressions, TArray<TNiagaraExprPtr>& OutputExpressions)
 {
 	OutputExpressions.Add(Expression_VMNative(VectorVM::EOp::length, InputExpressions));
+	return true;
 }
 
-void FNiagaraCompiler_VectorVM::Noise_Internal(TArray<TNiagaraExprPtr>& InputExpressions, TArray<TNiagaraExprPtr>& OutputExpressions)
+bool FNiagaraCompiler_VectorVM::Noise_Internal(TArray<TNiagaraExprPtr>& InputExpressions, TArray<TNiagaraExprPtr>& OutputExpressions)
 {
 	OutputExpressions.Add(Expression_VMNative(VectorVM::EOp::noise, InputExpressions));
+	return true;
 }
 
-void FNiagaraCompiler_VectorVM::Scalar_Internal(TArray<TNiagaraExprPtr>& InputExpressions, TArray<TNiagaraExprPtr>& OutputExpressions)
+bool FNiagaraCompiler_VectorVM::Scalar_Internal(TArray<TNiagaraExprPtr>& InputExpressions, TArray<TNiagaraExprPtr>& OutputExpressions)
 {
-	GetX_Internal(InputExpressions,OutputExpressions);
+	GetX_Internal(InputExpressions, OutputExpressions);
+	return true;
 }
 
-void FNiagaraCompiler_VectorVM::Vector_Internal(TArray<TNiagaraExprPtr>& InputExpressions, TArray<TNiagaraExprPtr>& OutputExpressions)
+bool FNiagaraCompiler_VectorVM::Vector_Internal(TArray<TNiagaraExprPtr>& InputExpressions, TArray<TNiagaraExprPtr>& OutputExpressions)
 {
 	OutputExpressions.Add(InputExpressions[0]);
+	return true;
 }
 
-void FNiagaraCompiler_VectorVM::Matrix_Internal(TArray<TNiagaraExprPtr>& InputExpressions, TArray<TNiagaraExprPtr>& OutputExpressions)
+bool FNiagaraCompiler_VectorVM::Matrix_Internal(TArray<TNiagaraExprPtr>& InputExpressions, TArray<TNiagaraExprPtr>& OutputExpressions)
 {
 	//Just collect the matrix rows into an expression for what ever use subsequent expressions have.
 	TNiagaraExprPtr MatrixExpr = Expression_Collection(InputExpressions);
 	MatrixExpr->Result.Type = ENiagaraDataType::Matrix;
 	OutputExpressions.Add(MatrixExpr);
+	return true;
 }
 
-void FNiagaraCompiler_VectorVM::Compose_Internal(TArray<TNiagaraExprPtr>& InputExpressions, TArray<TNiagaraExprPtr>& OutputExpressions)
+bool FNiagaraCompiler_VectorVM::Compose_Internal(TArray<TNiagaraExprPtr>& InputExpressions, TArray<TNiagaraExprPtr>& OutputExpressions)
 {
 	OutputExpressions.Add(Expression_VMNative(VectorVM::EOp::compose, InputExpressions));
+	return true;
 }
 
-void FNiagaraCompiler_VectorVM::ComposeX_Internal(TArray<TNiagaraExprPtr>& InputExpressions, TArray<TNiagaraExprPtr>& OutputExpressions)
+bool FNiagaraCompiler_VectorVM::ComposeX_Internal(TArray<TNiagaraExprPtr>& InputExpressions, TArray<TNiagaraExprPtr>& OutputExpressions)
 {
 	OutputExpressions.Add(Expression_VMNative(VectorVM::EOp::composex, InputExpressions));
+	return true;
 }
 
-void FNiagaraCompiler_VectorVM::ComposeY_Internal(TArray<TNiagaraExprPtr>& InputExpressions, TArray<TNiagaraExprPtr>& OutputExpressions)
+bool FNiagaraCompiler_VectorVM::ComposeY_Internal(TArray<TNiagaraExprPtr>& InputExpressions, TArray<TNiagaraExprPtr>& OutputExpressions)
 {
 	OutputExpressions.Add(Expression_VMNative(VectorVM::EOp::composey, InputExpressions));
+	return true;
 }
 
-void FNiagaraCompiler_VectorVM::ComposeZ_Internal(TArray<TNiagaraExprPtr>& InputExpressions, TArray<TNiagaraExprPtr>& OutputExpressions)
+bool FNiagaraCompiler_VectorVM::ComposeZ_Internal(TArray<TNiagaraExprPtr>& InputExpressions, TArray<TNiagaraExprPtr>& OutputExpressions)
 {
 	OutputExpressions.Add(Expression_VMNative(VectorVM::EOp::composez, InputExpressions));
+	return true;
 }
 
-void FNiagaraCompiler_VectorVM::ComposeW_Internal(TArray<TNiagaraExprPtr>& InputExpressions, TArray<TNiagaraExprPtr>& OutputExpressions)
+bool FNiagaraCompiler_VectorVM::ComposeW_Internal(TArray<TNiagaraExprPtr>& InputExpressions, TArray<TNiagaraExprPtr>& OutputExpressions)
 {
 	OutputExpressions.Add(Expression_VMNative(VectorVM::EOp::composew, InputExpressions));
+	return true;
 }
 
-void FNiagaraCompiler_VectorVM::Split_Internal(TArray<TNiagaraExprPtr>& InputExpressions, TArray<TNiagaraExprPtr>& OutputExpressions)
+bool FNiagaraCompiler_VectorVM::Split_Internal(TArray<TNiagaraExprPtr>& InputExpressions, TArray<TNiagaraExprPtr>& OutputExpressions)
 {
 	GetX_Internal(InputExpressions, OutputExpressions);
 	GetY_Internal(InputExpressions, OutputExpressions);
 	GetZ_Internal(InputExpressions, OutputExpressions);
 	GetW_Internal(InputExpressions, OutputExpressions);
+	return true;
 }
 
-void FNiagaraCompiler_VectorVM::GetX_Internal(TArray<TNiagaraExprPtr>& InputExpressions, TArray<TNiagaraExprPtr>& OutputExpressions)
+bool FNiagaraCompiler_VectorVM::GetX_Internal(TArray<TNiagaraExprPtr>& InputExpressions, TArray<TNiagaraExprPtr>& OutputExpressions)
 {
 	OutputExpressions.Add(Expression_VMNative(VectorVM::EOp::splatx, InputExpressions));
+	return true;
 }
 
-void FNiagaraCompiler_VectorVM::GetY_Internal(TArray<TNiagaraExprPtr>& InputExpressions, TArray<TNiagaraExprPtr>& OutputExpressions)
+bool FNiagaraCompiler_VectorVM::GetY_Internal(TArray<TNiagaraExprPtr>& InputExpressions, TArray<TNiagaraExprPtr>& OutputExpressions)
 {
 	OutputExpressions.Add(Expression_VMNative(VectorVM::EOp::splaty, InputExpressions));
+	return true;
 }
 
-void FNiagaraCompiler_VectorVM::GetZ_Internal(TArray<TNiagaraExprPtr>& InputExpressions, TArray<TNiagaraExprPtr>& OutputExpressions)
+bool FNiagaraCompiler_VectorVM::GetZ_Internal(TArray<TNiagaraExprPtr>& InputExpressions, TArray<TNiagaraExprPtr>& OutputExpressions)
 {
 	OutputExpressions.Add(Expression_VMNative(VectorVM::EOp::splatz, InputExpressions));
+	return true;
 }
 
-void FNiagaraCompiler_VectorVM::GetW_Internal(TArray<TNiagaraExprPtr>& InputExpressions, TArray<TNiagaraExprPtr>& OutputExpressions)
+bool FNiagaraCompiler_VectorVM::GetW_Internal(TArray<TNiagaraExprPtr>& InputExpressions, TArray<TNiagaraExprPtr>& OutputExpressions)
 {
 	OutputExpressions.Add(Expression_VMNative(VectorVM::EOp::splatw, InputExpressions));
+	return true;
 }
 
-void FNiagaraCompiler_VectorVM::TransformVector_Internal(TArray<TNiagaraExprPtr>& InputExpressions, TArray<TNiagaraExprPtr>& OutputExpressions)
+bool FNiagaraCompiler_VectorVM::TransformVector_Internal(TArray<TNiagaraExprPtr>& InputExpressions, TArray<TNiagaraExprPtr>& OutputExpressions)
 {
 	TNiagaraExprPtr MatrixExpr = InputExpressions[0];
 	check(MatrixExpr->SourceExpressions.Num() == 4);
@@ -574,9 +653,10 @@ void FNiagaraCompiler_VectorVM::TransformVector_Internal(TArray<TNiagaraExprPtr>
 	Temp[0] = Expression_VMNative(VectorVM::EOp::add, Temp[0], Temp[1]);
 	Temp[1] = Expression_VMNative(VectorVM::EOp::add, Temp[2], Temp[3]);
 	OutputExpressions.Add(Expression_VMNative(VectorVM::EOp::add, Temp[0], Temp[1]));
+	return true;
 }
 
-void FNiagaraCompiler_VectorVM::Transpose_Internal(TArray<TNiagaraExprPtr>& InputExpressions, TArray<TNiagaraExprPtr>& OutputExpressions)
+bool FNiagaraCompiler_VectorVM::Transpose_Internal(TArray<TNiagaraExprPtr>& InputExpressions, TArray<TNiagaraExprPtr>& OutputExpressions)
 {
 	TNiagaraExprPtr MatrixExpr = InputExpressions[0];
 	check(MatrixExpr->SourceExpressions.Num() == 4);
@@ -594,47 +674,53 @@ void FNiagaraCompiler_VectorVM::Transpose_Internal(TArray<TNiagaraExprPtr>& Inpu
 	Columns.Add(Expression_VMNative(VectorVM::EOp::composew, Row0, Row1, Row2, Row3));
 
 	Matrix_Internal(Columns, OutputExpressions);
+	return true;
 }
 
 
-void FNiagaraCompiler_VectorVM::Inverse_Internal(TArray<TNiagaraExprPtr>& InputExpressions, TArray<TNiagaraExprPtr>& OutputExpressions)
+bool FNiagaraCompiler_VectorVM::Inverse_Internal(TArray<TNiagaraExprPtr>& InputExpressions, TArray<TNiagaraExprPtr>& OutputExpressions)
 {
-	//OutputExpressions.Add(Expression_UnknownError());
-
-	
+	Error(LOCTEXT("VMNoImpl", "VM Operation not yet implemented"), nullptr, nullptr);
+	return false;
 }
 
 
-void FNiagaraCompiler_VectorVM::LessThan_Internal(TArray<TNiagaraExprPtr>& InputExpressions, TArray<TNiagaraExprPtr>& OutputExpressions)
+bool FNiagaraCompiler_VectorVM::LessThan_Internal(TArray<TNiagaraExprPtr>& InputExpressions, TArray<TNiagaraExprPtr>& OutputExpressions)
 {
 	OutputExpressions.Add(Expression_VMNative(VectorVM::EOp::lessthan, InputExpressions));
+	return true;
 }
 
 
-void FNiagaraCompiler_VectorVM::Sample_Internal(TArray<TNiagaraExprPtr>& InputExpressions, TArray<TNiagaraExprPtr>& OutputExpressions)
+bool FNiagaraCompiler_VectorVM::Sample_Internal(TArray<TNiagaraExprPtr>& InputExpressions, TArray<TNiagaraExprPtr>& OutputExpressions)
 {
 	OutputExpressions.Add(Expression_VMNative(VectorVM::EOp::sample, InputExpressions));
+	return true;
 }
 
-void FNiagaraCompiler_VectorVM::Write_Internal(TArray<TNiagaraExprPtr>& InputExpressions, TArray<TNiagaraExprPtr>& OutputExpressions)
+bool FNiagaraCompiler_VectorVM::Write_Internal(TArray<TNiagaraExprPtr>& InputExpressions, TArray<TNiagaraExprPtr>& OutputExpressions)
 {
 	OutputExpressions.Add(Expression_VMNative(VectorVM::EOp::bufferwrite, InputExpressions));
+	return true;
 }
 
-void FNiagaraCompiler_VectorVM::EventBroadcast_Internal(TArray<TNiagaraExprPtr>& InputExpressions, TArray<TNiagaraExprPtr>& OutputExpressions)
+bool FNiagaraCompiler_VectorVM::EventBroadcast_Internal(TArray<TNiagaraExprPtr>& InputExpressions, TArray<TNiagaraExprPtr>& OutputExpressions)
 {
 	OutputExpressions.Add(Expression_VMNative(VectorVM::EOp::eventbroadcast, InputExpressions));
+	return true;
 }
 
-void FNiagaraCompiler_VectorVM::EaseIn_Internal(TArray<TNiagaraExprPtr>& InputExpressions, TArray<TNiagaraExprPtr>& OutputExpressions)
+bool FNiagaraCompiler_VectorVM::EaseIn_Internal(TArray<TNiagaraExprPtr>& InputExpressions, TArray<TNiagaraExprPtr>& OutputExpressions)
 {
 	OutputExpressions.Add(Expression_VMNative(VectorVM::EOp::easein, InputExpressions));
+	return true;
 }
 
 
-void FNiagaraCompiler_VectorVM::EaseInOut_Internal(TArray<TNiagaraExprPtr>& InputExpressions, TArray<TNiagaraExprPtr>& OutputExpressions)
+bool FNiagaraCompiler_VectorVM::EaseInOut_Internal(TArray<TNiagaraExprPtr>& InputExpressions, TArray<TNiagaraExprPtr>& OutputExpressions)
 {
 	OutputExpressions.Add(Expression_VMNative(VectorVM::EOp::easeinout, InputExpressions));
+	return true;
 }
 
 
diff --git a/Engine/Source/Editor/NiagaraEditor/Private/NiagaraCompiler_VectorVM.h b/Engine/Source/Editor/NiagaraEditor/Private/NiagaraCompiler_VectorVM.h
--- a/Engine/Source/Editor/NiagaraEditor/Private/NiagaraCompiler_VectorVM.h
+++ b/Engine/Source/Editor/NiagaraEditor/Private/NiagaraCompiler_VectorVM.h
@@ -33,9 +33,9 @@ public:
 	void WriteCode(uint8 InCode);
 
 	//Begin INiagaraCompiler Interface	
-	virtual void CompileScript(UNiagaraScript* InScript) override;
+	virtual bool CompileScript(UNiagaraScript* InScript) override;
 
-#define NiagaraOp(OpName) virtual void OpName##_Internal(TArray<TNiagaraExprPtr>& InputExpressions, TArray<TNiagaraExprPtr>& OutputExpressions)override;
+#define NiagaraOp(OpName) virtual bool OpName##_Internal(TArray<TNiagaraExprPtr>& InputExpressions, TArray<TNiagaraExprPtr>& OutputExpressions)override;
 		NiagaraOpList;
 #undef NiagaraOp
 
diff --git a/Engine/Source/Editor/NiagaraEditor/Private/NiagaraEditor.cpp b/Engine/Source/Editor/NiagaraEditor/Private/NiagaraEditor.cpp
--- a/Engine/Source/Editor/NiagaraEditor/Private/NiagaraEditor.cpp
+++ b/Engine/Source/Editor/NiagaraEditor/Private/NiagaraEditor.cpp
@@ -10,6 +10,8 @@
 
 #include "Editor/PropertyEditor/Public/PropertyEditorModule.h"
 #include "Editor/PropertyEditor/Public/IDetailsView.h"
+#include "ScopedTransaction.h"
+#include "BlueprintEditorUtils.h"
  
 #define LOCTEXT_NAMESPACE "NiagaraEditor"
 
@@ -282,19 +284,26 @@ void FNiagaraEditor::DeleteSelectedNodes()
 		return;
 	}
 
-	const FGraphPanelSelectionSet SelectedNodes = NodeGraphEditor->GetSelectedNodes();
-	NodeGraphEditor->ClearSelectionSet();
-
-	for (FGraphPanelSelectionSet::TConstIterator NodeIt( SelectedNodes ); NodeIt; ++NodeIt)
 	{
-		if (UEdGraphNode* Node = Cast<UEdGraphNode>(*NodeIt))
+		const FScopedTransaction Transaction(FGenericCommands::Get().Delete->GetDescription());
+
+		NodeGraphEditor->GetCurrentGraph()->Modify();
+
+		const FGraphPanelSelectionSet SelectedNodes = NodeGraphEditor->GetSelectedNodes();
+
+		for (FGraphPanelSelectionSet::TConstIterator NodeIt(SelectedNodes); NodeIt; ++NodeIt)
 		{
-			if (Node->CanUserDeleteNode())
+			if (UEdGraphNode* Node = Cast<UEdGraphNode>(*NodeIt))
 			{
-				Node->DestroyNode();
+				if (Node->CanUserDeleteNode())
+				{
+					FBlueprintEditorUtils::RemoveNode(NULL, Node, true);
+				}
 			}
 		}
+		NodeGraphEditor->ClearSelectionSet();
 	}
+	NodeGraphEditor->NotifyGraphChanged();
 }
 
 bool FNiagaraEditor::CanDeleteNodes() const
diff --git a/Engine/Source/Editor/UnrealEd/Classes/Niagara/EdGraphSchema_Niagara.h b/Engine/Source/Editor/UnrealEd/Classes/Niagara/EdGraphSchema_Niagara.h
--- a/Engine/Source/Editor/UnrealEd/Classes/Niagara/EdGraphSchema_Niagara.h
+++ b/Engine/Source/Editor/UnrealEd/Classes/Niagara/EdGraphSchema_Niagara.h
@@ -59,9 +59,13 @@ class UEdGraphSchema_Niagara : public UEdGraphSchema
 
 	// Begin EdGraphSchema interface
 	virtual void GetGraphContextActions(FGraphContextMenuBuilder& ContextMenuBuilder) const override;
+	virtual void GetContextMenuActions(const UEdGraph* CurrentGraph, const UEdGraphNode* InGraphNode, const UEdGraphPin* InGraphPin, class FMenuBuilder* MenuBuilder, bool bIsDebugging) const override;
 	virtual const FPinConnectionResponse CanCreateConnection(const UEdGraphPin* A, const UEdGraphPin* B) const override;
 	virtual FLinearColor GetPinTypeColor(const FEdGraphPinType& PinType) const override;
 	virtual bool ShouldHidePinDefaultValue(UEdGraphPin* Pin) const override;
+	virtual bool TryCreateConnection(UEdGraphPin* A, UEdGraphPin* B) const override;
+
+	virtual void BreakSinglePinLink(UEdGraphPin* SourcePin, UEdGraphPin* TargetPin) override;
 	// End EdGraphSchema interface
 
 	UNREALED_API ENiagaraDataType GetPinType(UEdGraphPin* Pin)const;
@@ -75,5 +79,8 @@ class UEdGraphSchema_Niagara : public UEdGraphSchema
 	static const FLinearColor NodeTitleColor_Constant;
 	static const FLinearColor NodeTitleColor_SystemConstant;
 	static const FLinearColor NodeTitleColor_FunctionCall;
+
+private:
+	void GetBreakLinkToSubMenuActions(class FMenuBuilder& MenuBuilder, UEdGraphPin* InGraphPin);
 };
 
diff --git a/Engine/Source/Editor/UnrealEd/Classes/Niagara/NiagaraGraph.h b/Engine/Source/Editor/UnrealEd/Classes/Niagara/NiagaraGraph.h
--- a/Engine/Source/Editor/UnrealEd/Classes/Niagara/NiagaraGraph.h
+++ b/Engine/Source/Editor/UnrealEd/Classes/Niagara/NiagaraGraph.h
@@ -8,6 +8,10 @@ UCLASS(MinimalAPI)
 class UNiagaraGraph : public UEdGraph
 {
 	GENERATED_UCLASS_BODY()
+
+	// Begin UObject interface
+	virtual void PostEditChangeProperty(struct FPropertyChangedEvent& PropertyChangedEvent)override;
+	// End UObjet interface
 	
 	/** Get the source that owns this graph */
 	class UNiagaraScriptSource* GetSource() const;
diff --git a/Engine/Source/Editor/UnrealEd/Classes/Niagara/NiagaraNode.h b/Engine/Source/Editor/UnrealEd/Classes/Niagara/NiagaraNode.h
--- a/Engine/Source/Editor/UnrealEd/Classes/Niagara/NiagaraNode.h
+++ b/Engine/Source/Editor/UnrealEd/Classes/Niagara/NiagaraNode.h
@@ -9,6 +9,13 @@ class UNREALED_API UNiagaraNode : public UEdGraphNode
 	GENERATED_UCLASS_BODY()
 
 public:
+	// Begin UObject interface
+	virtual void PostEditChangeProperty(struct FPropertyChangedEvent& PropertyChangedEvent)override;
+	// End UObject interface
+
+	// Begin EdGraphNode interface
+	virtual void AutowireNewNode(UEdGraphPin* FromPin)override;
+	// End EdGraphNode interface
 
 	/** Get the Niagara graph that owns this node */
 	const class UNiagaraGraph* GetNiagaraGraph()const;
diff --git a/Engine/Source/Editor/UnrealEd/Classes/Niagara/NiagaraNodeInput.h b/Engine/Source/Editor/UnrealEd/Classes/Niagara/NiagaraNodeInput.h
--- a/Engine/Source/Editor/UnrealEd/Classes/Niagara/NiagaraNodeInput.h
+++ b/Engine/Source/Editor/UnrealEd/Classes/Niagara/NiagaraNodeInput.h
@@ -42,6 +42,7 @@ public:
 	virtual void AllocateDefaultPins() override;
 	virtual FText GetNodeTitle(ENodeTitleType::Type TitleType) const override;
 	virtual FLinearColor GetNodeTitleColor() const override;
+	virtual void AutowireNewNode(UEdGraphPin* FromPin)override;
 	// End EdGraphNode interface
 
 	virtual void Compile(class INiagaraCompiler* Compiler, TArray<FNiagaraNodeResult>& Outputs)override;
diff --git a/Engine/Source/Editor/UnrealEd/Private/Niagara/EdGraphSchema_Niagara.cpp b/Engine/Source/Editor/UnrealEd/Private/Niagara/EdGraphSchema_Niagara.cpp
--- a/Engine/Source/Editor/UnrealEd/Private/Niagara/EdGraphSchema_Niagara.cpp
+++ b/Engine/Source/Editor/UnrealEd/Private/Niagara/EdGraphSchema_Niagara.cpp
@@ -5,9 +5,11 @@
 #include "BlueprintGraphDefinitions.h"
 #include "GraphEditorSettings.h"
 #include "NiagaraComponent.h"
+#include "GraphEditorActions.h"
 
 #include "INiagaraEditor.h"
 #include "INiagaraCompiler.h"
+#include "ScopedTransaction.h"
 
 #define LOCTEXT_NAMESPACE "NiagaraSchema"
 
@@ -31,6 +33,9 @@ UEdGraphNode* FNiagaraSchemaAction_NewNode::PerformAction(class UEdGraph* Parent
 	// If there is a template, we actually use it
 	if (NodeTemplate != NULL)
 	{
+		const FScopedTransaction Transaction(NSLOCTEXT("UnrealEd", "NiagaraEditorNewNode", "Niagara Editor: New Node"));
+		ParentGraph->Modify();
+
 		NodeTemplate->SetFlags(RF_Transactional);
 
 		// set outer to be the graph so it doesn't go away
@@ -63,6 +68,8 @@ UEdGraphNode* FNiagaraSchemaAction_NewNode::PerformAction(class UEdGraph* Parent
 		NodeTemplate->SnapToGrid(SNAP_GRID);
 
 		ResultNode = NodeTemplate;
+
+		ParentGraph->NotifyGraphChanged();
 	}
 
 	return ResultNode;
@@ -238,6 +245,20 @@ const FPinConnectionResponse UEdGraphSchema_Niagara::CanCreateConnection(const U
 	}
 }
 
+void UEdGraphSchema_Niagara::BreakSinglePinLink(UEdGraphPin* SourcePin, UEdGraphPin* TargetPin) 
+{
+	const FScopedTransaction Transaction(NSLOCTEXT("UnrealEd", "NiagaraEditorBreakConnection", "Niagara Editor: Break Connection"));
+
+	Super::BreakSinglePinLink(SourcePin, TargetPin);
+}
+
+bool UEdGraphSchema_Niagara::TryCreateConnection(UEdGraphPin* A, UEdGraphPin* B) const
+{
+	const FScopedTransaction Transaction(NSLOCTEXT("UnrealEd", "NiagaraEditorCreateConnection", "Niagara Editor: Create Connection"));
+
+	return Super::TryCreateConnection(A, B);
+}
+
 FLinearColor UEdGraphSchema_Niagara::GetPinTypeColor(const FEdGraphPinType& PinType) const
 {
 	const FString& TypeString = PinType.PinCategory;
@@ -342,4 +363,81 @@ bool UEdGraphSchema_Niagara::IsSystemConstant(const FNiagaraVariableInfo& Variab
 	return UNiagaraComponent::GetSystemConstants().Find(Variable) != INDEX_NONE;
 }
 
+void UEdGraphSchema_Niagara::GetBreakLinkToSubMenuActions(class FMenuBuilder& MenuBuilder, UEdGraphPin* InGraphPin)
+{
+	// Make sure we have a unique name for every entry in the list
+	TMap< FString, uint32 > LinkTitleCount;
+
+	// Add all the links we could break from
+	for (TArray<class UEdGraphPin*>::TConstIterator Links(InGraphPin->LinkedTo); Links; ++Links)
+	{
+		UEdGraphPin* Pin = *Links;
+		FString TitleString = Pin->GetOwningNode()->GetNodeTitle(ENodeTitleType::ListView).ToString();
+		FText Title = FText::FromString(TitleString);
+		if (Pin->PinName != TEXT(""))
+		{
+			TitleString = FString::Printf(TEXT("%s (%s)"), *TitleString, *Pin->PinName);
+
+			// Add name of connection if possible
+			FFormatNamedArguments Args;
+			Args.Add(TEXT("NodeTitle"), Title);
+			Args.Add(TEXT("PinName"), Pin->GetDisplayName());
+			Title = FText::Format(LOCTEXT("BreakDescPin", "{NodeTitle} ({PinName})"), Args);
+		}
+
+		uint32 &Count = LinkTitleCount.FindOrAdd(TitleString);
+
+		FText Description;
+		FFormatNamedArguments Args;
+		Args.Add(TEXT("NodeTitle"), Title);
+		Args.Add(TEXT("NumberOfNodes"), Count);
+
+		if (Count == 0)
+		{
+			Description = FText::Format(LOCTEXT("BreakDesc", "Break link to {NodeTitle}"), Args);
+		}
+		else
+		{
+			Description = FText::Format(LOCTEXT("BreakDescMulti", "Break link to {NodeTitle} ({NumberOfNodes})"), Args);
+		}
+		++Count;
+		MenuBuilder.AddMenuEntry(Description, Description, FSlateIcon(), FUIAction(
+			FExecuteAction::CreateUObject((UEdGraphSchema_Niagara*const)this, &UEdGraphSchema::BreakSinglePinLink, const_cast<UEdGraphPin*>(InGraphPin), *Links)));
+	}
+}
+
+void UEdGraphSchema_Niagara::GetContextMenuActions(const UEdGraph* CurrentGraph, const UEdGraphNode* InGraphNode, const UEdGraphPin* InGraphPin, class FMenuBuilder* MenuBuilder, bool bIsDebugging) const
+{
+	if (InGraphPin)
+	{
+		MenuBuilder->BeginSection("EdGraphSchema_NiagaraPinActions", LOCTEXT("PinActionsMenuHeader", "Pin Actions"));
+		{
+			// Only display the 'Break Link' option if there is a link to break!
+			if (InGraphPin->LinkedTo.Num() > 0)
+			{
+				MenuBuilder->AddMenuEntry(FGraphEditorCommands::Get().BreakPinLinks);
+
+				// add sub menu for break link to
+				if (InGraphPin->LinkedTo.Num() > 1)
+				{
+					MenuBuilder->AddSubMenu(
+						LOCTEXT("BreakLinkTo", "Break Link To..."),
+						LOCTEXT("BreakSpecificLinks", "Break a specific link..."),
+						FNewMenuDelegate::CreateUObject((UEdGraphSchema_Niagara*const)this, &UEdGraphSchema_Niagara::GetBreakLinkToSubMenuActions, const_cast<UEdGraphPin*>(InGraphPin)));
+				}
+				else
+				{
+					((UEdGraphSchema_Niagara*const)this)->GetBreakLinkToSubMenuActions(*MenuBuilder, const_cast<UEdGraphPin*>(InGraphPin));
+				}				
+			}
+		}
+		MenuBuilder->EndSection();
+	}
+	else if (InGraphNode)
+	{
+	}
+
+	Super::GetContextMenuActions(CurrentGraph, InGraphNode, InGraphPin, MenuBuilder, bIsDebugging);
+}
+
 #undef LOCTEXT_NAMESPACE
\ No newline at end of file
diff --git a/Engine/Source/Editor/UnrealEd/Private/Niagara/NiagaraEditorCommon.cpp b/Engine/Source/Editor/UnrealEd/Private/Niagara/NiagaraEditorCommon.cpp
--- a/Engine/Source/Editor/UnrealEd/Private/Niagara/NiagaraEditorCommon.cpp
+++ b/Engine/Source/Editor/UnrealEd/Private/Niagara/NiagaraEditorCommon.cpp
@@ -614,7 +614,7 @@ TMap<FName, FNiagaraOpInfo>& FNiagaraOpInfo::GetOpInfoMap()
 
 //Debugging variant that checks inputs and outputs match expectations of FNiagaraOpInfo.
 #if DO_CHECK
-#define NiagaraOp(OpName) void INiagaraCompiler::OpName(INiagaraCompiler* Compiler, TArray<TNiagaraExprPtr>& InputExpressions, TArray<TNiagaraExprPtr>& OutputExpressions)\
+#define NiagaraOp(OpName) bool INiagaraCompiler::OpName(INiagaraCompiler* Compiler, TArray<TNiagaraExprPtr>& InputExpressions, TArray<TNiagaraExprPtr>& OutputExpressions)\
 {\
 	Compiler->CheckInputs(FNiagaraOpInfo::OpName, InputExpressions); \
 	return Compiler->OpName##_Internal(InputExpressions, OutputExpressions); \
@@ -625,7 +625,7 @@ NiagaraOpList
 
 #else//DO_CHECK
 
-#define NiagaraOp(OpName) void INiagaraCompiler::OpName(INiagaraCompiler* Compiler, TArray<TNiagaraExprPtr>& InputExpressions, TArray<TNiagaraExprPtr>& OutputExpressions)\
+#define NiagaraOp(OpName) bool INiagaraCompiler::OpName(INiagaraCompiler* Compiler, TArray<TNiagaraExprPtr>& InputExpressions, TArray<TNiagaraExprPtr>& OutputExpressions)\
 { \
 	return Compiler->OpName##_Internal(InputExpressions, OutputExpressions); \
 }
diff --git a/Engine/Source/Editor/UnrealEd/Private/Niagara/NiagaraNode.cpp b/Engine/Source/Editor/UnrealEd/Private/Niagara/NiagaraNode.cpp
--- a/Engine/Source/Editor/UnrealEd/Private/Niagara/NiagaraNode.cpp
+++ b/Engine/Source/Editor/UnrealEd/Private/Niagara/NiagaraNode.cpp
@@ -7,6 +7,44 @@ UNiagaraNode::UNiagaraNode(const FObjectInitializer& ObjectInitializer)
 {
 }
 
+void UNiagaraNode::PostEditChangeProperty(struct FPropertyChangedEvent& PropertyChangedEvent)
+{
+	Super::PostEditChangeProperty(PropertyChangedEvent);
+
+	GetGraph()->NotifyGraphChanged();
+}
+
+void UNiagaraNode::AutowireNewNode(UEdGraphPin* FromPin)
+{
+	Super::AutowireNewNode(FromPin);
+
+	if (FromPin != nullptr)
+	{
+		const UEdGraphSchema_Niagara* Schema = CastChecked<UEdGraphSchema_Niagara>(GetSchema());
+		check(Schema);
+		
+		ENiagaraDataType FromType = Schema->GetPinType(FromPin);
+
+		//Find first of this nodes pins with the right type and direction.
+		UEdGraphPin* FirstPinOfSameType = NULL;
+		EEdGraphPinDirection DesiredDirection = FromPin->Direction == EGPD_Output ? EGPD_Input : EGPD_Output;
+		for (UEdGraphPin* Pin : Pins)
+		{
+			ENiagaraDataType ToType = Schema->GetPinType(Pin);
+			if (FromType == ToType && Pin->Direction == DesiredDirection)
+			{
+				FirstPinOfSameType = Pin;
+				break;
+			}			
+		}
+
+		if (FirstPinOfSameType && GetSchema()->TryCreateConnection(FromPin, FirstPinOfSameType))
+		{
+			FromPin->GetOwningNode()->NodeConnectionListChanged();
+		}
+	}
+}
+
 const UNiagaraGraph* UNiagaraNode::GetNiagaraGraph()const
 {
 	return CastChecked<UNiagaraGraph>(GetGraph());
diff --git a/Engine/Source/Editor/UnrealEd/Private/Niagara/NiagaraNodeFunctionCall.cpp b/Engine/Source/Editor/UnrealEd/Private/Niagara/NiagaraNodeFunctionCall.cpp
--- a/Engine/Source/Editor/UnrealEd/Private/Niagara/NiagaraNodeFunctionCall.cpp
+++ b/Engine/Source/Editor/UnrealEd/Private/Niagara/NiagaraNodeFunctionCall.cpp
@@ -18,154 +18,153 @@ void UNiagaraNodeFunctionCall::PostLoad()
 void UNiagaraNodeFunctionCall::PostEditChangeProperty(struct FPropertyChangedEvent& PropertyChangedEvent)
 {
 	ReallocatePins();
+	Super::PostEditChangeProperty(PropertyChangedEvent);
 }
 
 void UNiagaraNodeFunctionCall::ReallocatePins()
 {
-	if (!FunctionScript)
-		return;
-
-	const UEdGraphSchema_Niagara* Schema = GetDefault<UEdGraphSchema_Niagara>();
-	Modify();
-
-	// Break any links to 'orphan' pins
-	for (int32 PinIndex = 0; PinIndex < Pins.Num(); ++PinIndex)
+	if (FunctionScript)
 	{
-		UEdGraphPin* Pin = Pins[PinIndex];
-		TArray<class UEdGraphPin*>& LinkedToRef = Pin->LinkedTo;
-		for (int32 LinkIdx = 0; LinkIdx < LinkedToRef.Num(); LinkIdx++)
+		const UEdGraphSchema_Niagara* Schema = GetDefault<UEdGraphSchema_Niagara>();
+		Modify();
+
+		// Break any links to 'orphan' pins
+		for (int32 PinIndex = 0; PinIndex < Pins.Num(); ++PinIndex)
 		{
-			UEdGraphPin* OtherPin = LinkedToRef[LinkIdx];
-			// If we are linked to a pin that its owner doesn't know about, break that link
-			if (!OtherPin->GetOwningNode()->Pins.Contains(OtherPin))
+			UEdGraphPin* Pin = Pins[PinIndex];
+			TArray<class UEdGraphPin*>& LinkedToRef = Pin->LinkedTo;
+			for (int32 LinkIdx = 0; LinkIdx < LinkedToRef.Num(); LinkIdx++)
 			{
-				Pin->LinkedTo.Remove(OtherPin);
+				UEdGraphPin* OtherPin = LinkedToRef[LinkIdx];
+				// If we are linked to a pin that its owner doesn't know about, break that link
+				if (!OtherPin->GetOwningNode()->Pins.Contains(OtherPin))
+				{
+					Pin->LinkedTo.Remove(OtherPin);
+				}
 			}
 		}
-	}
 
-	// Store the old Input and Output pins
-	TArray<UEdGraphPin*> OldInputPins;
-	TArray<UEdGraphPin*> OldOutputPins;
-	GetInputPins(OldInputPins);
-	GetOutputPins(OldOutputPins);
-
-	// Move the existing pins to a saved array
-	TArray<UEdGraphPin*> OldPins(Pins);
-	Pins.Empty();
-	
-	if (UNiagaraScriptSource* Source = Cast<UNiagaraScriptSource>(FunctionScript->Source))
-	{
-		if (UNiagaraGraph* Graph = Cast<UNiagaraGraph>(Source->NodeGraph))
+		// Store the old Input and Output pins
+		TArray<UEdGraphPin*> OldInputPins;
+		TArray<UEdGraphPin*> OldOutputPins;
+		GetInputPins(OldInputPins);
+		GetOutputPins(OldOutputPins);
+
+		// Move the existing pins to a saved array
+		TArray<UEdGraphPin*> OldPins(Pins);
+		Pins.Empty();
+
+		if (UNiagaraScriptSource* Source = Cast<UNiagaraScriptSource>(FunctionScript->Source))
 		{
-			if (UNiagaraNodeOutput* OutNode = Graph->FindOutputNode())
+			if (UNiagaraGraph* Graph = Cast<UNiagaraGraph>(Source->NodeGraph))
 			{
-				//Add the input pins.
-				TArray<UNiagaraNodeInput*> InputNodes;
-				Graph->FindInputNodes(InputNodes);
-
-				for (UNiagaraNodeInput* InNode : InputNodes)
+				if (UNiagaraNodeOutput* OutNode = Graph->FindOutputNode())
 				{
-					//Don't expose system constants as inputs.
-					if (Schema->IsSystemConstant(InNode->Input))
-						continue;
+					//Add the input pins.
+					TArray<UNiagaraNodeInput*> InputNodes;
+					Graph->FindInputNodes(InputNodes);
 
-					check(InNode);
-					switch (InNode->Input.Type)
-					{
-					case ENiagaraDataType::Scalar:
-					{
-						UEdGraphPin* NewPin = CreatePin(EGPD_Input, Schema->PC_Float, TEXT(""), NULL, false, false, InNode->Input.Name.ToString());
-						NewPin->bDefaultValueIsIgnored = true;
-					}
-						break;
-					case ENiagaraDataType::Vector:
+					for (UNiagaraNodeInput* InNode : InputNodes)
 					{
-						UEdGraphPin* NewPin = CreatePin(EGPD_Input, Schema->PC_Vector, TEXT(""), NULL, false, false, InNode->Input.Name.ToString());
-						NewPin->bDefaultValueIsIgnored = true;
+						//Don't expose system constants as inputs.
+						if (Schema->IsSystemConstant(InNode->Input))
+							continue;
+
+						check(InNode);
+						switch (InNode->Input.Type)
+						{
+						case ENiagaraDataType::Scalar:
+						{
+														 UEdGraphPin* NewPin = CreatePin(EGPD_Input, Schema->PC_Float, TEXT(""), NULL, false, false, InNode->Input.Name.ToString());
+														 NewPin->bDefaultValueIsIgnored = true;
+						}
+							break;
+						case ENiagaraDataType::Vector:
+						{
+														 UEdGraphPin* NewPin = CreatePin(EGPD_Input, Schema->PC_Vector, TEXT(""), NULL, false, false, InNode->Input.Name.ToString());
+														 NewPin->bDefaultValueIsIgnored = true;
+						}
+							break;
+						case ENiagaraDataType::Matrix:
+						{
+														 UEdGraphPin* NewPin = CreatePin(EGPD_Input, Schema->PC_Matrix, TEXT(""), NULL, false, false, InNode->Input.Name.ToString());
+														 NewPin->bDefaultValueIsIgnored = true;
+						}
+							break;
+						};
 					}
-						break;
-					case ENiagaraDataType::Matrix:
-					{
-						UEdGraphPin* NewPin = CreatePin(EGPD_Input, Schema->PC_Matrix, TEXT(""), NULL, false, false, InNode->Input.Name.ToString());
-						NewPin->bDefaultValueIsIgnored = true;
-					}
-						break;
-					};
-				}
 
-				//Add the output pins.
-				for (FNiagaraVariableInfo& Output : OutNode->Outputs)
-				{
-					switch (Output.Type)
+					//Add the output pins.
+					for (FNiagaraVariableInfo& Output : OutNode->Outputs)
 					{
-					case ENiagaraDataType::Scalar:
-					{						
-						UEdGraphPin* NewPin = CreatePin(EGPD_Output, Schema->PC_Float, TEXT(""), NULL, false, false, Output.Name.ToString());
-						NewPin->bDefaultValueIsIgnored = true;
+						switch (Output.Type)
+						{
+						case ENiagaraDataType::Scalar:
+						{
+														 UEdGraphPin* NewPin = CreatePin(EGPD_Output, Schema->PC_Float, TEXT(""), NULL, false, false, Output.Name.ToString());
+														 NewPin->bDefaultValueIsIgnored = true;
+						}
+							break;
+						case ENiagaraDataType::Vector:
+						{
+														 UEdGraphPin* NewPin = CreatePin(EGPD_Output, Schema->PC_Vector, TEXT(""), NULL, false, false, Output.Name.ToString());
+														 NewPin->bDefaultValueIsIgnored = true;
+						}
+							break;
+						case ENiagaraDataType::Matrix:
+						{
+														 UEdGraphPin* NewPin = CreatePin(EGPD_Output, Schema->PC_Matrix, TEXT(""), NULL, false, false, Output.Name.ToString());
+														 NewPin->bDefaultValueIsIgnored = true;
+						}
+							break;
+						};
 					}
-						break;
-					case ENiagaraDataType::Vector:
-					{						
-						UEdGraphPin* NewPin = CreatePin(EGPD_Output, Schema->PC_Vector, TEXT(""), NULL, false, false, Output.Name.ToString());
-						NewPin->bDefaultValueIsIgnored = true;
-					}
-						break;
-					case ENiagaraDataType::Matrix:
-					{
-						UEdGraphPin* NewPin = CreatePin(EGPD_Output, Schema->PC_Matrix, TEXT(""), NULL, false, false, Output.Name.ToString());
-						NewPin->bDefaultValueIsIgnored = true;
-					}
-						break;
-					};
-				}		
+				}
 			}
 		}
-	}
 
-	// Get new Input and Output pins
-	TArray<UEdGraphPin*> NewInputPins;
-	TArray<UEdGraphPin*> NewOutputPins;
-	GetInputPins(NewInputPins);
-	GetOutputPins(NewOutputPins);
+		// Get new Input and Output pins
+		TArray<UEdGraphPin*> NewInputPins;
+		TArray<UEdGraphPin*> NewOutputPins;
+		GetInputPins(NewInputPins);
+		GetOutputPins(NewOutputPins);
 
-	for (int32 PinIndex = 0; PinIndex < OldInputPins.Num(); PinIndex++)
-	{
-		if (PinIndex < NewInputPins.Num())
+		for (int32 PinIndex = 0; PinIndex < OldInputPins.Num(); PinIndex++)
 		{
-			NewInputPins[PinIndex]->CopyPersistentDataFromOldPin(*OldInputPins[PinIndex]);
+			if (PinIndex < NewInputPins.Num())
+			{
+				NewInputPins[PinIndex]->CopyPersistentDataFromOldPin(*OldInputPins[PinIndex]);
+			}
 		}
-	}
 
-	for (int32 PinIndex = 0; PinIndex < OldOutputPins.Num(); PinIndex++)
-	{
-		if (PinIndex < NewOutputPins.Num())
+		for (int32 PinIndex = 0; PinIndex < OldOutputPins.Num(); PinIndex++)
 		{
-			NewOutputPins[PinIndex]->CopyPersistentDataFromOldPin(*OldOutputPins[PinIndex]);
+			if (PinIndex < NewOutputPins.Num())
+			{
+				NewOutputPins[PinIndex]->CopyPersistentDataFromOldPin(*OldOutputPins[PinIndex]);
+			}
 		}
-	}
 
-	OldInputPins.Empty();
-	OldOutputPins.Empty();
+		OldInputPins.Empty();
+		OldOutputPins.Empty();
 
-	// Throw away the original pins
-	for (int32 OldPinIndex = 0; OldPinIndex < OldPins.Num(); ++OldPinIndex)
-	{
-		UEdGraphPin* OldPin = OldPins[OldPinIndex];
-		OldPin->Modify();
-		OldPin->BreakAllPinLinks();
+		// Throw away the original pins
+		for (int32 OldPinIndex = 0; OldPinIndex < OldPins.Num(); ++OldPinIndex)
+		{
+			UEdGraphPin* OldPin = OldPins[OldPinIndex];
+			OldPin->Modify();
+			OldPin->BreakAllPinLinks();
 
 #if 0
-		UEdGraphNode::ReturnPinToPool(OldPin);
+			UEdGraphNode::ReturnPinToPool(OldPin);
 #else
-		OldPin->Rename(NULL, GetTransientPackage(), REN_None);
-		OldPin->RemoveFromRoot();
-		OldPin->MarkPendingKill();
+			OldPin->Rename(NULL, GetTransientPackage(), REN_None);
+			OldPin->RemoveFromRoot();
+			OldPin->MarkPendingKill();
 #endif
+		}
+		OldPins.Empty();
 	}
-	OldPins.Empty();
-
-	GetGraph()->NotifyGraphChanged();
 }
 
 void UNiagaraNodeFunctionCall::AllocateDefaultPins()
diff --git a/Engine/Source/Editor/UnrealEd/Private/Niagara/NiagaraNodeInput.cpp b/Engine/Source/Editor/UnrealEd/Private/Niagara/NiagaraNodeInput.cpp
--- a/Engine/Source/Editor/UnrealEd/Private/Niagara/NiagaraNodeInput.cpp
+++ b/Engine/Source/Editor/UnrealEd/Private/Niagara/NiagaraNodeInput.cpp
@@ -18,6 +18,7 @@ UNiagaraNodeInput::UNiagaraNodeInput(const FObjectInitializer& ObjectInitializer
 void UNiagaraNodeInput::PostEditChangeProperty(struct FPropertyChangedEvent& PropertyChangedEvent)
 {
 	ReallocatePins();
+	Super::PostEditChangeProperty(PropertyChangedEvent);
 }
 
 void UNiagaraNodeInput::ReallocatePins()
@@ -117,8 +118,6 @@ void UNiagaraNodeInput::ReallocatePins()
 #endif
 	}
 	OldPins.Empty();
-
-	GetGraph()->NotifyGraphChanged();
 }
 
 void UNiagaraNodeInput::AllocateDefaultPins()
@@ -128,7 +127,7 @@ void UNiagaraNodeInput::AllocateDefaultPins()
 
 FText UNiagaraNodeInput::GetNodeTitle(ENodeTitleType::Type TitleType) const
 {
-	return Input.Name == NAME_None ? NSLOCTEXT("NiagaraNodeInput", "Input", "Input") : FText::FromName(Input.Name);
+	return NSLOCTEXT("NiagaraNodeInput", "Input", "Input");
 }
 
 FLinearColor UNiagaraNodeInput::GetNodeTitleColor() const
@@ -146,6 +145,27 @@ FLinearColor UNiagaraNodeInput::GetNodeTitleColor() const
 	}
 }
 
+void UNiagaraNodeInput::AutowireNewNode(UEdGraphPin* FromPin)
+{
+	if (FromPin != nullptr)
+	{
+		const UEdGraphSchema_Niagara* Schema = CastChecked<UEdGraphSchema_Niagara>(GetSchema());
+		check(Schema);
+		if (Input.Name == NAME_None)
+		{
+			Input.Name = FName(*FromPin->PinName);
+			Input.Type = Schema->GetPinType(FromPin);
+			ReallocatePins();
+		}
+		check(Pins.Num() == 1 && Pins[0] != NULL);
+		
+		if (GetSchema()->TryCreateConnection(FromPin, Pins[0]))
+		{
+			FromPin->GetOwningNode()->NodeConnectionListChanged();
+		}
+	}
+}
+
 void UNiagaraNodeInput::Compile(class INiagaraCompiler* Compiler, TArray<FNiagaraNodeResult>& Outputs)
 {
 	//Input nodes in scripts being used as functions by other scripts will actually never get this far.
diff --git a/Engine/Source/Editor/UnrealEd/Private/Niagara/NiagaraNodeOp.cpp b/Engine/Source/Editor/UnrealEd/Private/Niagara/NiagaraNodeOp.cpp
--- a/Engine/Source/Editor/UnrealEd/Private/Niagara/NiagaraNodeOp.cpp
+++ b/Engine/Source/Editor/UnrealEd/Private/Niagara/NiagaraNodeOp.cpp
@@ -4,6 +4,8 @@
 #include "BlueprintGraphDefinitions.h"
 #include "GraphEditorSettings.h"
 
+#define LOCTEXT_NAMESPACE "NiagaraNodeOp"
+
 UNiagaraNodeOp::UNiagaraNodeOp(const FObjectInitializer& ObjectInitializer)
 	: Super(ObjectInitializer)
 {
@@ -123,22 +125,42 @@ void UNiagaraNodeOp::Compile(class INiagaraCompiler* Compiler, TArray<FNiagaraNo
 	int32 NumOutputs = OpInfo->Outputs.Num();
 
 	TArray<TNiagaraExprPtr> Inputs;
+	bool bError = false;
 	for (int32 i = 0; i < NumInputs; ++i)
 	{
 		UEdGraphPin *Pin = Pins[i];
 		check(Pin->Direction == EGPD_Input);
-		Inputs.Add(Compiler->CompilePin(Pin));
+		TNiagaraExprPtr InputExpr = Compiler->CompilePin(Pin);
+		if (!InputExpr.IsValid())
+		{
+			bError = true;
+			FFormatNamedArguments Args;
+			Args.Add(TEXT("OpName"), GetNodeTitle(ENodeTitleType::FullTitle));
+			FText Format = LOCTEXT("InputErrorFormat", "Error compiling input on {OpName} node.");
+			Compiler->Error(FText::Format(Format, Args), this, Pin);
+		}
+		Inputs.Add(InputExpr);
 	}
 	
 	TArray<TNiagaraExprPtr> OutputExpressions;
-	OpInfo->OpDelegate.Execute(Compiler, Inputs, OutputExpressions);
-
-	check( OutputExpressions.Num() == OpInfo->Outputs.Num() );
+	if ( !bError && OpInfo->OpDelegate.Execute(Compiler, Inputs, OutputExpressions))
+	{
+		check(OutputExpressions.Num() == OpInfo->Outputs.Num());
 
-	for (int32 i = 0; i < NumOutputs; ++i)
+		for (int32 i = 0; i < NumOutputs; ++i)
+		{
+			UEdGraphPin *Pin = Pins[NumInputs + i];
+			check(Pin->Direction == EGPD_Output);
+			Outputs.Add(FNiagaraNodeResult(OutputExpressions[i], Pin));
+		}
+	}
+	else
 	{
-		UEdGraphPin *Pin = Pins[NumInputs + i];
-		check(Pin->Direction == EGPD_Output);
-		Outputs.Add(FNiagaraNodeResult(OutputExpressions[i], Pin));
+		FFormatNamedArguments Args;
+		Args.Add(TEXT("OpName"), GetNodeTitle(ENodeTitleType::FullTitle));
+		FText Format = LOCTEXT("InputErrorFormat", "Error compiling {OpName} node.");
+		Compiler->Error(FText::Format(Format, Args), this, nullptr);
 	}
-}
\ No newline at end of file
+}
+
+#undef  LOCTEXT_NAMESPACE
\ No newline at end of file
diff --git a/Engine/Source/Editor/UnrealEd/Private/Niagara/NiagaraNodeOutput.cpp b/Engine/Source/Editor/UnrealEd/Private/Niagara/NiagaraNodeOutput.cpp
--- a/Engine/Source/Editor/UnrealEd/Private/Niagara/NiagaraNodeOutput.cpp
+++ b/Engine/Source/Editor/UnrealEd/Private/Niagara/NiagaraNodeOutput.cpp
@@ -4,6 +4,8 @@
 #include "BlueprintGraphDefinitions.h"
 #include "GraphEditorSettings.h"
 
+#define LOCTEXT_NAMESPACE "NiagaraNodeOutput"
+
 UNiagaraNodeOutput::UNiagaraNodeOutput(const FObjectInitializer& ObjectInitializer)
 : Super(ObjectInitializer)
 {
@@ -12,6 +14,7 @@ UNiagaraNodeOutput::UNiagaraNodeOutput(const FObjectInitializer& ObjectInitializ
 void UNiagaraNodeOutput::PostEditChangeProperty(struct FPropertyChangedEvent& PropertyChangedEvent)
 {
 	ReallocatePins();
+	Super::PostEditChangeProperty(PropertyChangedEvent);
 }
 
 void UNiagaraNodeOutput::ReallocatePins()
@@ -137,6 +140,7 @@ void UNiagaraNodeOutput::Compile(class INiagaraCompiler* Compiler, TArray<FNiaga
 {
 	TArray<TNiagaraExprPtr> InputExpressions;
 
+	bool bError = false;
 	for (int32 i = 0; i < Outputs.Num(); ++i)
 	{
 		const FNiagaraVariableInfo& Out = Outputs[i];
@@ -146,16 +150,21 @@ void UNiagaraNodeOutput::Compile(class INiagaraCompiler* Compiler, TArray<FNiaga
 		TNiagaraExprPtr Result = Compiler->CompilePin(Pin);
 		if (!Result.IsValid())
 		{
-			//The pin was not connected so pass through the initial value of the attribute.
-			Result = Compiler->GetAttribute(Out);
+			bError = true;
+			Compiler->Error(FText::Format(LOCTEXT("OutputErrorFormat", "Error compiling attriubte {0}"), Pin->PinFriendlyName), this, Pin);
 		}
 
 		InputExpressions.Add(Result);
 	}
 
-	check(Outputs.Num() == InputExpressions.Num());
-	for (int32 i = 0; i < Outputs.Num(); ++i)
+	if (!bError)
 	{
-		OutputExpressions.Add(FNiagaraNodeResult(Compiler->Output(Outputs[i], InputExpressions[i]), Pins[i]));
+		check(Outputs.Num() == InputExpressions.Num());
+		for (int32 i = 0; i < Outputs.Num(); ++i)
+		{
+			OutputExpressions.Add(FNiagaraNodeResult(Compiler->Output(Outputs[i], InputExpressions[i]), Pins[i]));
+		}
 	}
-}
\ No newline at end of file
+}
+
+#undef LOCTEXT_NAMESPACE
\ No newline at end of file
diff --git a/Engine/Source/Editor/UnrealEd/Private/Niagara/NiagaraScriptSource.cpp b/Engine/Source/Editor/UnrealEd/Private/Niagara/NiagaraScriptSource.cpp
--- a/Engine/Source/Editor/UnrealEd/Private/Niagara/NiagaraScriptSource.cpp
+++ b/Engine/Source/Editor/UnrealEd/Private/Niagara/NiagaraScriptSource.cpp
@@ -15,6 +15,11 @@ UNiagaraGraph::UNiagaraGraph(const FObjectInitializer& ObjectInitializer)
 	Schema = UEdGraphSchema_Niagara::StaticClass();
 }
 
+void UNiagaraGraph::PostEditChangeProperty(FPropertyChangedEvent& PropertyChangedEvent)
+{
+	NotifyGraphChanged();
+}
+
 class UNiagaraScriptSource* UNiagaraGraph::GetSource() const
 {
 	return CastChecked<UNiagaraScriptSource>(GetOuter());
diff --git a/Engine/Source/Editor/UnrealEd/Public/INiagaraCompiler.h b/Engine/Source/Editor/UnrealEd/Public/INiagaraCompiler.h
--- a/Engine/Source/Editor/UnrealEd/Public/INiagaraCompiler.h
+++ b/Engine/Source/Editor/UnrealEd/Public/INiagaraCompiler.h
@@ -4,13 +4,15 @@
 
 #include "NiagaraEditorCommon.h"
 
+class UNiagaraNode;
+
 //Interface for Niagara compilers.
 class UNREALED_API INiagaraCompiler
 {
 public:
 
 	/** Compiles a script. */
-	virtual void CompileScript(class UNiagaraScript* InScript) = 0;
+	virtual bool CompileScript(class UNiagaraScript* InScript) = 0;
 
 	/** Traverses the pins input, recursively compiling them and then adds any expressions needed for the owner node.*/
 	virtual TNiagaraExprPtr CompilePin(UEdGraphPin* Pin) = 0;
@@ -36,18 +38,24 @@ public:
 	virtual TNiagaraExprPtr Output(const FNiagaraVariableInfo& Attr, TNiagaraExprPtr& SourceExpression) = 0;
 
 	/** Checks the validity of the input expressions against the informations about this operation in FNiagaraOpInfo. */
-	virtual void CheckInputs(FName OpName, TArray<TNiagaraExprPtr>& Inputs) = 0;
+	virtual bool CheckInputs(FName OpName, TArray<TNiagaraExprPtr>& Inputs) = 0;
 	/** Checks the validity of the output expressions against the informations about this operation in FNiagaraOpInfo. */
-	virtual void CheckOutputs(FName OpName, TArray<TNiagaraExprPtr>& Outputs) = 0;
+	virtual bool CheckOutputs(FName OpName, TArray<TNiagaraExprPtr>& Outputs) = 0;
+
+	/** Adds an error to be reported to the user. Any error will lead to compilation failure. */
+	virtual void Error(FText ErrorText, UNiagaraNode* Node, UEdGraphPin* Pin) = 0 ;
+
+	/** Adds a warning to be reported to the user. Warnings will not cause a compilation failure. */
+	virtual void Warning(FText WarningText, UNiagaraNode* Node, UEdGraphPin* Pin) = 0;
 
 	//Static function named after each op that allows access to internals via FNiagaraOpInfo.Func.
-#define NiagaraOp(OpName) static void OpName(INiagaraCompiler* Compiler, TArray<TNiagaraExprPtr>& InputExpressions, TArray<TNiagaraExprPtr>& OutputExpressions);
+#define NiagaraOp(OpName) static bool OpName(INiagaraCompiler* Compiler, TArray<TNiagaraExprPtr>& InputExpressions, TArray<TNiagaraExprPtr>& OutputExpressions);
 	NiagaraOpList;
 #undef NiagaraOp
 
 protected:
 	//Internal virtual function for each op. 
-#define NiagaraOp(OpName) virtual void OpName##_Internal(TArray<TNiagaraExprPtr>& InputExpressions, TArray<TNiagaraExprPtr>& OutputExpressions) = 0;
+#define NiagaraOp(OpName) virtual bool OpName##_Internal(TArray<TNiagaraExprPtr>& InputExpressions, TArray<TNiagaraExprPtr>& OutputExpressions) = 0;
 	NiagaraOpList;
 #undef NiagaraOp
 
diff --git a/Engine/Source/Editor/UnrealEd/Public/NiagaraEditorCommon.h b/Engine/Source/Editor/UnrealEd/Public/NiagaraEditorCommon.h
--- a/Engine/Source/Editor/UnrealEd/Public/NiagaraEditorCommon.h
+++ b/Engine/Source/Editor/UnrealEd/Public/NiagaraEditorCommon.h
@@ -130,7 +130,7 @@ public:
 	TArray<TNiagaraExprPtr> SourceExpressions;
 };
 
-DECLARE_DELEGATE_ThreeParams(FNiagaraOpDelegate, class INiagaraCompiler*, TArray<TNiagaraExprPtr>&, TArray<TNiagaraExprPtr> &);
+DECLARE_DELEGATE_RetVal_ThreeParams(bool, FNiagaraOpDelegate, class INiagaraCompiler*, TArray<TNiagaraExprPtr>&, TArray<TNiagaraExprPtr> &);
 
 /** Information about an input or output of a Niagara operation node. */
 class UNREALED_API FNiagaraOpInOutInfo
