Support for Cut/Copy/Paste, Duplicate and Select All in Niagara graphs.
Plus better support for Undo/Redo.

diff --git a/Engine/Source/Editor/NiagaraEditor/Private/NiagaraEditor.cpp b/Engine/Source/Editor/NiagaraEditor/Private/NiagaraEditor.cpp
--- a/Engine/Source/Editor/NiagaraEditor/Private/NiagaraEditor.cpp
+++ b/Engine/Source/Editor/NiagaraEditor/Private/NiagaraEditor.cpp
@@ -141,9 +141,34 @@ TSharedRef<SGraphEditor> FNiagaraEditor::CreateGraphEditorWidget(UEdGraph* InGra
 		GraphEditorCommands = MakeShareable( new FUICommandList );
 	
 		// Editing commands
-		GraphEditorCommands->MapAction( FGenericCommands::Get().Delete,
-			FExecuteAction::CreateSP( this, &FNiagaraEditor::DeleteSelectedNodes ),
-			FCanExecuteAction::CreateSP( this, &FNiagaraEditor::CanDeleteNodes )
+		GraphEditorCommands->MapAction(FGenericCommands::Get().SelectAll,
+			FExecuteAction::CreateRaw(this, &FNiagaraEditor::SelectAllNodes),
+			FCanExecuteAction::CreateRaw(this, &FNiagaraEditor::CanSelectAllNodes)
+			);
+
+		GraphEditorCommands->MapAction(FGenericCommands::Get().Delete,
+			FExecuteAction::CreateRaw(this, &FNiagaraEditor::DeleteSelectedNodes),
+			FCanExecuteAction::CreateRaw(this, &FNiagaraEditor::CanDeleteNodes)
+			);
+
+		GraphEditorCommands->MapAction(FGenericCommands::Get().Copy,
+			FExecuteAction::CreateRaw(this, &FNiagaraEditor::CopySelectedNodes),
+			FCanExecuteAction::CreateRaw(this, &FNiagaraEditor::CanCopyNodes)
+			);
+
+		GraphEditorCommands->MapAction(FGenericCommands::Get().Cut,
+			FExecuteAction::CreateRaw(this, &FNiagaraEditor::CutSelectedNodes),
+			FCanExecuteAction::CreateRaw(this, &FNiagaraEditor::CanCutNodes)
+			);
+
+		GraphEditorCommands->MapAction(FGenericCommands::Get().Paste,
+			FExecuteAction::CreateRaw(this, &FNiagaraEditor::PasteNodes),
+			FCanExecuteAction::CreateRaw(this, &FNiagaraEditor::CanPasteNodes)
+			);
+
+		GraphEditorCommands->MapAction(FGenericCommands::Get().Duplicate,
+			FExecuteAction::CreateRaw(this, &FNiagaraEditor::DuplicateNodes),
+			FCanExecuteAction::CreateRaw(this, &FNiagaraEditor::CanDuplicateNodes)
 			);
 	}
 
@@ -276,61 +301,310 @@ FReply FNiagaraEditor::OnCompileClicked()
 	return FReply::Handled();
 }
 
-void FNiagaraEditor::DeleteSelectedNodes()
+void FNiagaraEditor::OnSelectedNodesChanged(const TSet<class UObject*>& NewSelection)
 {
-	TSharedPtr<SGraphEditor> NodeGraphEditor = NodeGraphEditorPtr.Pin();
-	if (!NodeGraphEditor.IsValid())
+	TArray<UObject*> SelectedObjects;
+
+	UObject* EditObject = Script;
+
+	if (NewSelection.Num() == 0)
 	{
-		return;
+		SelectedObjects.Add(EditObject);
+	}
+	else
+	{
+		for (TSet<class UObject*>::TConstIterator SetIt(NewSelection); SetIt; ++SetIt)
+		{
+			SelectedObjects.Add(*SetIt);
+		}
+	}
+
+	GetDetailView()->SetObjects(SelectedObjects, true);
+}
+
+void FNiagaraEditor::PostUndo(bool bSuccess)
+{
+	if (bSuccess)
+	{
+		// Clear selection, to avoid holding refs to nodes that go away
+		TSharedPtr<SGraphEditor> CurrentGraphEditor = NodeGraphEditorPtr.Pin();
+		if (CurrentGraphEditor.IsValid())
+		{
+			CurrentGraphEditor->ClearSelectionSet();
+			CurrentGraphEditor->NotifyGraphChanged();
+		}
+		FSlateApplication::Get().DismissAllMenus();
+	}
+}
+
+void FNiagaraEditor::PostRedo(bool bSuccess)
+{
+	if (bSuccess)
+	{
+		// Clear selection, to avoid holding refs to nodes that go away
+		TSharedPtr<SGraphEditor> CurrentGraphEditor = NodeGraphEditorPtr.Pin();
+		if (CurrentGraphEditor.IsValid())
+		{
+			CurrentGraphEditor->ClearSelectionSet();
+			CurrentGraphEditor->NotifyGraphChanged();
+		}
+		FSlateApplication::Get().DismissAllMenus();
+	}
+}
+
+void FNiagaraEditor::SelectAllNodes()
+{
+	TSharedPtr<SGraphEditor> CurrentGraphEditor = NodeGraphEditorPtr.Pin();
+	if (CurrentGraphEditor.IsValid())
+	{
+		CurrentGraphEditor->SelectAllNodes();
 	}
+}
 
+bool FNiagaraEditor::CanSelectAllNodes() const
+{
+	return true;
+}
+
+void FNiagaraEditor::DeleteSelectedNodes()
+{
+	TSharedPtr<SGraphEditor> CurrentGraphEditor = NodeGraphEditorPtr.Pin();
+	if (!CurrentGraphEditor.IsValid())
 	{
-		const FScopedTransaction Transaction(FGenericCommands::Get().Delete->GetDescription());
+		return;
+	}
 
-		NodeGraphEditor->GetCurrentGraph()->Modify();
+	const FScopedTransaction Transaction(FGenericCommands::Get().Delete->GetDescription());
+	CurrentGraphEditor->GetCurrentGraph()->Modify();
 
-		const FGraphPanelSelectionSet SelectedNodes = NodeGraphEditor->GetSelectedNodes();
+	const FGraphPanelSelectionSet SelectedNodes = CurrentGraphEditor->GetSelectedNodes();
+	CurrentGraphEditor->ClearSelectionSet();
 
-		for (FGraphPanelSelectionSet::TConstIterator NodeIt(SelectedNodes); NodeIt; ++NodeIt)
+	for (FGraphPanelSelectionSet::TConstIterator NodeIt(SelectedNodes); NodeIt; ++NodeIt)
+	{
+		if (UEdGraphNode* Node = Cast<UEdGraphNode>(*NodeIt))
 		{
-			if (UEdGraphNode* Node = Cast<UEdGraphNode>(*NodeIt))
+			if (Node->CanUserDeleteNode())
 			{
-				if (Node->CanUserDeleteNode())
-				{
-					FBlueprintEditorUtils::RemoveNode(NULL, Node, true);
-				}
+				Node->Modify();
+				Node->DestroyNode();
 			}
 		}
-		NodeGraphEditor->ClearSelectionSet();
 	}
-	NodeGraphEditor->NotifyGraphChanged();
 }
 
 bool FNiagaraEditor::CanDeleteNodes() const
 {
-	TSharedPtr<SGraphEditor> NodeGraphEditor = NodeGraphEditorPtr.Pin();
-	return (NodeGraphEditor.IsValid() && NodeGraphEditor->GetSelectedNodes().Num() > 0);
+	// If any of the nodes can be deleted then we should allow deleting
+	TSharedPtr<SGraphEditor> CurrentGraphEditor = NodeGraphEditorPtr.Pin();
+	const FGraphPanelSelectionSet SelectedNodes = CurrentGraphEditor->GetSelectedNodes();
+	for (FGraphPanelSelectionSet::TConstIterator SelectedIter(SelectedNodes); SelectedIter; ++SelectedIter)
+	{
+		UEdGraphNode* Node = Cast<UEdGraphNode>(*SelectedIter);
+		if (Node && Node->CanUserDeleteNode())
+		{
+			return true;
+		}
+	}
+
+	return false;
 }
 
-void FNiagaraEditor::OnSelectedNodesChanged(const TSet<class UObject*>& NewSelection)
+void FNiagaraEditor::DeleteSelectedDuplicatableNodes()
 {
-	TArray<UObject*> SelectedObjects;
+	TSharedPtr<SGraphEditor> CurrentGraphEditor = NodeGraphEditorPtr.Pin();
+	if (!CurrentGraphEditor.IsValid())
+	{
+		return;
+	}
 
-	UObject* EditObject = Script;
+	const FGraphPanelSelectionSet OldSelectedNodes = CurrentGraphEditor->GetSelectedNodes();
+	CurrentGraphEditor->ClearSelectionSet();
 
-	if (NewSelection.Num() == 0)
+	for (FGraphPanelSelectionSet::TConstIterator SelectedIter(OldSelectedNodes); SelectedIter; ++SelectedIter)
 	{
-		SelectedObjects.Add(EditObject);
+		UEdGraphNode* Node = Cast<UEdGraphNode>(*SelectedIter);
+		if (Node && Node->CanDuplicateNode())
+		{
+			CurrentGraphEditor->SetNodeSelection(Node, true);
+		}
 	}
-	else
+
+	// Delete the duplicatable nodes
+	DeleteSelectedNodes();
+
+	CurrentGraphEditor->ClearSelectionSet();
+
+	for (FGraphPanelSelectionSet::TConstIterator SelectedIter(OldSelectedNodes); SelectedIter; ++SelectedIter)
 	{
-		for (TSet<class UObject*>::TConstIterator SetIt(NewSelection); SetIt; ++SetIt)
+		if (UEdGraphNode* Node = Cast<UEdGraphNode>(*SelectedIter))
 		{
-			SelectedObjects.Add(*SetIt);
+			CurrentGraphEditor->SetNodeSelection(Node, true);
 		}
 	}
+}
 
-	GetDetailView()->SetObjects(SelectedObjects, true);
+void FNiagaraEditor::CutSelectedNodes()
+{
+	CopySelectedNodes();
+	DeleteSelectedDuplicatableNodes();
+}
+
+bool FNiagaraEditor::CanCutNodes() const
+{
+	return CanCopyNodes() && CanDeleteNodes();
+}
+
+void FNiagaraEditor::CopySelectedNodes()
+{
+	// Export the selected nodes and place the text on the clipboard
+	TSharedPtr<SGraphEditor> CurrentGraphEditor = NodeGraphEditorPtr.Pin();
+	FGraphPanelSelectionSet SelectedNodes = CurrentGraphEditor->GetSelectedNodes();
+
+	FString ExportedText;
+
+	int32 CopySubNodeIndex = 0;
+	for (FGraphPanelSelectionSet::TIterator SelectedIter(SelectedNodes); SelectedIter; ++SelectedIter)
+	{
+		UEdGraphNode* Node = Cast<UEdGraphNode>(*SelectedIter);
+		if (Node == nullptr)
+		{
+			SelectedIter.RemoveCurrent();
+			continue;
+		}
+
+		Node->PrepareForCopying();
+	}
+	
+	FEdGraphUtilities::ExportNodesToText(SelectedNodes, ExportedText);
+	FPlatformMisc::ClipboardCopy(*ExportedText);
+}
+
+bool FNiagaraEditor::CanCopyNodes() const
+{
+	// If any of the nodes can be duplicated then we should allow copying
+	TSharedPtr<SGraphEditor> CurrentGraphEditor = NodeGraphEditorPtr.Pin();
+	const FGraphPanelSelectionSet SelectedNodes = CurrentGraphEditor->GetSelectedNodes();
+	for (FGraphPanelSelectionSet::TConstIterator SelectedIter(SelectedNodes); SelectedIter; ++SelectedIter)
+	{
+		UEdGraphNode* Node = Cast<UEdGraphNode>(*SelectedIter);
+		if (Node && Node->CanDuplicateNode())
+		{
+			return true;
+		}
+	}
+
+	return false;
+}
+
+void FNiagaraEditor::PasteNodes()
+{
+	TSharedPtr<SGraphEditor> CurrentGraphEditor = NodeGraphEditorPtr.Pin();
+	if (CurrentGraphEditor.IsValid())
+	{
+		PasteNodesHere(CurrentGraphEditor->GetPasteLocation());
+	}
+}
+
+void FNiagaraEditor::PasteNodesHere(const FVector2D& Location)
+{
+	TSharedPtr<SGraphEditor> CurrentGraphEditor = NodeGraphEditorPtr.Pin();
+	if (!CurrentGraphEditor.IsValid())
+	{
+		return;
+	}
+
+	// Undo/Redo support
+	const FScopedTransaction Transaction(FGenericCommands::Get().Paste->GetDescription());
+	UEdGraph* EdGraph = Cast<UEdGraph>(CurrentGraphEditor->GetCurrentGraph());
+	EdGraph->Modify();
+
+	const FGraphPanelSelectionSet SelectedNodes = CurrentGraphEditor->GetSelectedNodes();
+
+	// Clear the selection set (newly pasted stuff will be selected)
+	CurrentGraphEditor->ClearSelectionSet();
+
+	// Grab the text to paste from the clipboard.
+	FString TextToImport;
+	FPlatformMisc::ClipboardPaste(TextToImport);
+
+	// Import the nodes
+	TSet<UEdGraphNode*> PastedNodes;
+	FEdGraphUtilities::ImportNodesFromText(EdGraph, TextToImport, /*out*/ PastedNodes);
+
+	//Average position of nodes so we can move them while still maintaining relative distances to each other
+	FVector2D AvgNodePosition(0.0f, 0.0f);
+
+	for (TSet<UEdGraphNode*>::TIterator It(PastedNodes); It; ++It)
+	{
+		UEdGraphNode* Node = (*It);
+		if (Node)
+		{
+			AvgNodePosition.X += Node->NodePosX;
+			AvgNodePosition.Y += Node->NodePosY;
+		}
+	}
+
+	if (PastedNodes.Num() > 0)
+	{
+		float InvNumNodes = 1.0f / float(PastedNodes.Num());
+		AvgNodePosition.X *= InvNumNodes;
+		AvgNodePosition.Y *= InvNumNodes;
+	}
+	
+	for (TSet<UEdGraphNode*>::TIterator It(PastedNodes); It; ++It)
+	{
+		UEdGraphNode* PasteNode = (*It);
+		if (PasteNode)
+		{
+			// Select the newly pasted stuff
+			CurrentGraphEditor->SetNodeSelection(PasteNode, true);
+
+			PasteNode->NodePosX = (PasteNode->NodePosX - AvgNodePosition.X) + Location.X;
+			PasteNode->NodePosY = (PasteNode->NodePosY - AvgNodePosition.Y) + Location.Y;
+
+			PasteNode->SnapToGrid(16);
+
+			// Give new node a different Guid from the old one
+			PasteNode->CreateNewGuid();
+		}
+	}
+
+
+	// Update UI
+	CurrentGraphEditor->NotifyGraphChanged();
+
+	UObject* GraphOwner = EdGraph->GetOuter();
+	if (GraphOwner)
+	{
+		GraphOwner->PostEditChange();
+		GraphOwner->MarkPackageDirty();
+	}
+}
+
+bool FNiagaraEditor::CanPasteNodes() const
+{
+	TSharedPtr<SGraphEditor> CurrentGraphEditor = NodeGraphEditorPtr.Pin();
+	if (!CurrentGraphEditor.IsValid())
+	{
+		return false;
+	}
+
+	FString ClipboardContent;
+	FPlatformMisc::ClipboardPaste(ClipboardContent);
+
+	return FEdGraphUtilities::CanImportNodesFromText(CurrentGraphEditor->GetCurrentGraph(), ClipboardContent);
+}
+
+void FNiagaraEditor::DuplicateNodes()
+{
+	CopySelectedNodes();
+	PasteNodes();
+}
+
+bool FNiagaraEditor::CanDuplicateNodes() const
+{
+	return CanCopyNodes();
 }
 
 #undef LOCTEXT_NAMESPACE
diff --git a/Engine/Source/Editor/NiagaraEditor/Private/NiagaraEditor.h b/Engine/Source/Editor/NiagaraEditor/Private/NiagaraEditor.h
--- a/Engine/Source/Editor/NiagaraEditor/Private/NiagaraEditor.h
+++ b/Engine/Source/Editor/NiagaraEditor/Private/NiagaraEditor.h
@@ -7,7 +7,7 @@
 
 
 /** Viewer/editor for a DataTable */
-class FNiagaraEditor : public INiagaraEditor, public FNotifyHook
+class FNiagaraEditor : public INiagaraEditor, public FNotifyHook, public FEditorUndoClient
 {
 
 public:
@@ -28,10 +28,10 @@ public:
 	virtual FLinearColor GetWorldCentricTabColorScale() const override;
 	// End IToolkit interface
 
-	// Delegates
-	FReply OnCompileClicked();
-	void DeleteSelectedNodes();
-	bool CanDeleteNodes() const;
+	// Begin FEditorUndoClient Interface
+	virtual void PostUndo(bool bSuccess) override;
+	virtual void PostRedo(bool bSuccess) override;
+	// End of FEditorUndoClient
 
 	// Widget Accessors
 	TSharedRef<class IDetailsView> GetDetailView() const { return NiagaraDetailsView.ToSharedRef(); }
@@ -52,6 +52,25 @@ private:
 
 	/** Builds the toolbar widget */
 	void ExtendToolbar();
+
+
+	// Delegates for graph editor commands
+	FReply OnCompileClicked();
+	void SelectAllNodes();
+	bool CanSelectAllNodes() const;
+	void DeleteSelectedNodes();
+	bool CanDeleteNodes() const;
+	void DeleteSelectedDuplicatableNodes();
+	void CutSelectedNodes();
+	bool CanCutNodes() const;
+	void CopySelectedNodes();
+	bool CanCopyNodes() const;
+	void PasteNodes();
+	void PasteNodesHere(const FVector2D& Location);
+	bool CanPasteNodes() const;
+	void DuplicateNodes();
+	bool CanDuplicateNodes() const;
+
 private:
 
 	/* The Script being edited */
diff --git a/Engine/Source/Editor/UnrealEd/Classes/Niagara/NiagaraNode.h b/Engine/Source/Editor/UnrealEd/Classes/Niagara/NiagaraNode.h
--- a/Engine/Source/Editor/UnrealEd/Classes/Niagara/NiagaraNode.h
+++ b/Engine/Source/Editor/UnrealEd/Classes/Niagara/NiagaraNode.h
@@ -9,9 +9,6 @@ class UNREALED_API UNiagaraNode : public UEdGraphNode
 	GENERATED_UCLASS_BODY()
 
 public:
-	// Begin UObject interface
-	virtual void PostEditChangeProperty(struct FPropertyChangedEvent& PropertyChangedEvent)override;
-	// End UObject interface
 
 	// Begin EdGraphNode interface
 	virtual void AutowireNewNode(UEdGraphPin* FromPin)override;
diff --git a/Engine/Source/Editor/UnrealEd/Private/Niagara/NiagaraNode.cpp b/Engine/Source/Editor/UnrealEd/Private/Niagara/NiagaraNode.cpp
--- a/Engine/Source/Editor/UnrealEd/Private/Niagara/NiagaraNode.cpp
+++ b/Engine/Source/Editor/UnrealEd/Private/Niagara/NiagaraNode.cpp
@@ -7,13 +7,6 @@ UNiagaraNode::UNiagaraNode(const FObjectInitializer& ObjectInitializer)
 {
 }
 
-void UNiagaraNode::PostEditChangeProperty(struct FPropertyChangedEvent& PropertyChangedEvent)
-{
-	Super::PostEditChangeProperty(PropertyChangedEvent);
-
-	GetGraph()->NotifyGraphChanged();
-}
-
 void UNiagaraNode::AutowireNewNode(UEdGraphPin* FromPin)
 {
 	Super::AutowireNewNode(FromPin);
diff --git a/Engine/Source/Editor/UnrealEd/Private/Niagara/NiagaraNodeFunctionCall.cpp b/Engine/Source/Editor/UnrealEd/Private/Niagara/NiagaraNodeFunctionCall.cpp
--- a/Engine/Source/Editor/UnrealEd/Private/Niagara/NiagaraNodeFunctionCall.cpp
+++ b/Engine/Source/Editor/UnrealEd/Private/Niagara/NiagaraNodeFunctionCall.cpp
@@ -17,7 +17,10 @@ void UNiagaraNodeFunctionCall::PostLoad()
 
 void UNiagaraNodeFunctionCall::PostEditChangeProperty(struct FPropertyChangedEvent& PropertyChangedEvent)
 {
-	ReallocatePins();
+	if (PropertyChangedEvent.Property != nullptr)
+	{
+		ReallocatePins();
+	}
 	Super::PostEditChangeProperty(PropertyChangedEvent);
 }
 
diff --git a/Engine/Source/Editor/UnrealEd/Private/Niagara/NiagaraNodeInput.cpp b/Engine/Source/Editor/UnrealEd/Private/Niagara/NiagaraNodeInput.cpp
--- a/Engine/Source/Editor/UnrealEd/Private/Niagara/NiagaraNodeInput.cpp
+++ b/Engine/Source/Editor/UnrealEd/Private/Niagara/NiagaraNodeInput.cpp
@@ -17,7 +17,11 @@ UNiagaraNodeInput::UNiagaraNodeInput(const FObjectInitializer& ObjectInitializer
 
 void UNiagaraNodeInput::PostEditChangeProperty(struct FPropertyChangedEvent& PropertyChangedEvent)
 {
-	ReallocatePins();
+	if (PropertyChangedEvent.Property != nullptr)
+	{
+		ReallocatePins();
+	}
+
 	Super::PostEditChangeProperty(PropertyChangedEvent);
 }
 
diff --git a/Engine/Source/Editor/UnrealEd/Private/Niagara/NiagaraNodeOutput.cpp b/Engine/Source/Editor/UnrealEd/Private/Niagara/NiagaraNodeOutput.cpp
--- a/Engine/Source/Editor/UnrealEd/Private/Niagara/NiagaraNodeOutput.cpp
+++ b/Engine/Source/Editor/UnrealEd/Private/Niagara/NiagaraNodeOutput.cpp
@@ -13,7 +13,10 @@ UNiagaraNodeOutput::UNiagaraNodeOutput(const FObjectInitializer& ObjectInitializ
 
 void UNiagaraNodeOutput::PostEditChangeProperty(struct FPropertyChangedEvent& PropertyChangedEvent)
 {
-	ReallocatePins();
+	if (PropertyChangedEvent.Property != nullptr)
+	{
+		ReallocatePins();
+	}
 	Super::PostEditChangeProperty(PropertyChangedEvent);
 }
 
