Compile fix

diff --git a/Engine/Source/Editor/NiagaraEditor/Private/NiagaraEditor.cpp b/Engine/Source/Editor/NiagaraEditor/Private/NiagaraEditor.cpp
--- a/Engine/Source/Editor/NiagaraEditor/Private/NiagaraEditor.cpp
+++ b/Engine/Source/Editor/NiagaraEditor/Private/NiagaraEditor.cpp
@@ -7,6 +7,7 @@
 #include "Editor/WorkspaceMenuStructure/Public/WorkspaceMenuStructureModule.h"
 #include "SDockTab.h"
 #include "GenericCommands.h"
+#include "EdGraphUtilities.h"
 
 #include "Editor/PropertyEditor/Public/PropertyEditorModule.h"
 #include "Editor/PropertyEditor/Public/IDetailsView.h"
