Niagara: first pass embedded sequencer timeline

diff --git a/Engine/Source/Editor/NiagaraEditor/NiagaraEditor.Build.cs b/Engine/Source/Editor/NiagaraEditor/NiagaraEditor.Build.cs
--- a/Engine/Source/Editor/NiagaraEditor/NiagaraEditor.Build.cs
+++ b/Engine/Source/Editor/NiagaraEditor/NiagaraEditor.Build.cs
@@ -20,14 +20,16 @@ public class NiagaraEditor : ModuleRules
 				"UnrealEd", 
 				"GraphEditor", 
 				"VectorVM",
-                "Niagara"
+                "Niagara",
+                "MovieSceneCore",
+				"Sequencer",
 			}
 		);
 
         PublicDependencyModuleNames.AddRange(
             new string[] {
 				"Engine",
-                "Niagara"
+                "Niagara",
             }
         );
 
@@ -39,6 +41,10 @@ public class NiagaraEditor : ModuleRules
 				"LevelEditor"}
                 );
 
-		DynamicallyLoadedModuleNames.Add("WorkspaceMenuStructure");
+		DynamicallyLoadedModuleNames.AddRange(
+            new string[] {
+                "WorkspaceMenuStructure",
+                }
+            );
 	}
 }
diff --git a/Engine/Source/Editor/NiagaraEditor/Private/NiagaraEffectEditor.cpp b/Engine/Source/Editor/NiagaraEditor/Private/NiagaraEffectEditor.cpp
--- a/Engine/Source/Editor/NiagaraEditor/Private/NiagaraEffectEditor.cpp
+++ b/Engine/Source/Editor/NiagaraEditor/Private/NiagaraEffectEditor.cpp
@@ -5,13 +5,16 @@
 
 #include "Toolkits/IToolkitHost.h"
 #include "Editor/WorkspaceMenuStructure/Public/WorkspaceMenuStructureModule.h"
-#include "SNiagaraEffectEditorWidget.h"
 #include "SDockTab.h"
 
 #define LOCTEXT_NAMESPACE "NiagaraEffectEditor"
 
 
 const FName FNiagaraEffectEditor::UpdateTabId(TEXT("NiagaraEditor_Effect"));
+const FName FNiagaraEffectEditor::ViewportTabID(TEXT("NiagaraEffectEditor_Viewport"));
+const FName FNiagaraEffectEditor::EmitterEditorTabID(TEXT("NiagaraEffectEditor_EmitterEditor"));
+const FName FNiagaraEffectEditor::CurveEditorTabID(TEXT("NiagaraEffectEditor_CurveEditor"));
+const FName FNiagaraEffectEditor::SequencerTabID(TEXT("NiagaraEffectEditor_Sequencer"));
 
 void FNiagaraEffectEditor::RegisterTabSpawners(const TSharedRef<class FTabManager>& TabManager)
 {
@@ -19,9 +22,29 @@ void FNiagaraEffectEditor::RegisterTabSpawners(const TSharedRef<class FTabManage
 
 	FAssetEditorToolkit::RegisterTabSpawners(TabManager);
 
+	TabManager->RegisterTabSpawner(ViewportTabID, FOnSpawnTab::CreateSP(this, &FNiagaraEffectEditor::SpawnTab_Viewport))
+		.SetDisplayName(LOCTEXT("Preview", "Preview"))
+		.SetGroup(WorkspaceMenuCategory.ToSharedRef())
+		.SetIcon(FSlateIcon(FEditorStyle::GetStyleSetName(), "LevelEditor.Tabs.Viewports"));
+
+	TabManager->RegisterTabSpawner(EmitterEditorTabID, FOnSpawnTab::CreateSP(this, &FNiagaraEffectEditor::SpawnTab_EmitterList))
+		.SetDisplayName(LOCTEXT("Emitters", "Emitters"))
+		.SetGroup(WorkspaceMenuCategory.ToSharedRef());
+
+	TabManager->RegisterTabSpawner(CurveEditorTabID, FOnSpawnTab::CreateSP(this, &FNiagaraEffectEditor::SpawnTab_CurveEd))
+		.SetDisplayName(LOCTEXT("Curves", "Curves"))
+		.SetGroup(WorkspaceMenuCategory.ToSharedRef());
+
+	TabManager->RegisterTabSpawner(SequencerTabID, FOnSpawnTab::CreateSP(this, &FNiagaraEffectEditor::SpawnTab_Sequencer))
+		.SetDisplayName(LOCTEXT("Timeline", "Timeline"))
+		.SetGroup(WorkspaceMenuCategory.ToSharedRef());
+
+
 	TabManager->RegisterTabSpawner(UpdateTabId, FOnSpawnTab::CreateSP(this, &FNiagaraEffectEditor::SpawnTab))
 		.SetDisplayName(LOCTEXT("NiagaraEffect", "Niagara Effect"))
 		.SetGroup( WorkspaceMenuCategory.ToSharedRef() );
+
+
 }
 
 void FNiagaraEffectEditor::UnregisterTabSpawners(const TSharedRef<class FTabManager>& TabManager)
@@ -29,8 +52,15 @@ void FNiagaraEffectEditor::UnregisterTabSpawners(const TSharedRef<class FTabMana
 	FAssetEditorToolkit::UnregisterTabSpawners(TabManager);
 
 	TabManager->UnregisterTabSpawner(UpdateTabId);
+	TabManager->UnregisterTabSpawner(ViewportTabID);
+	TabManager->UnregisterTabSpawner(EmitterEditorTabID);
+	TabManager->UnregisterTabSpawner(CurveEditorTabID);
+	TabManager->UnregisterTabSpawner(SequencerTabID);
 }
 
+
+
+
 FNiagaraEffectEditor::~FNiagaraEffectEditor()
 {
 
@@ -41,23 +71,67 @@ void FNiagaraEffectEditor::InitNiagaraEffectEditor(const EToolkitMode::Type Mode
 {
 	Effect = InEffect;
 	check(Effect != NULL);
+	EffectInstance = new FNiagaraEffectInstance(InEffect);
+
+	if (!Sequencer.IsValid())
+	{
+		FSequencerViewParams ViewParams;
+		ViewParams.InitalViewRange = TRange<float>(-0.02f, 3.2f);
+		ViewParams.InitialScrubPosition = 0;
+
+		SequencerBindingManager = MakeShareable(new FNiagaraSequencerObjectBindingManager());
+
+		MovieScene = NewObject<UMovieScene>(InEffect, FName("Niagara Effect MovieScene"), RF_RootSet);
 
-	TSharedRef<FTabManager::FLayout> StandaloneDefaultLayout = FTabManager::NewLayout("Standalone_Niagara_Effect_Layout_v2")
+		ISequencerModule &SeqModule = FModuleManager::LoadModuleChecked< ISequencerModule >("Sequencer");
+		FDelegateHandle CreateTrackEditorHandle = SeqModule.RegisterTrackEditor_Handle(FOnCreateTrackEditor::CreateStatic(&FNiagaraEffectEditor::CreateTrackEditor));
+		Sequencer = SeqModule.CreateSequencer(MovieScene, ViewParams, SequencerBindingManager.ToSharedRef());
+
+		for (TSharedPtr<FNiagaraSimulation> Emitter : EffectInstance->GetEmitters())
+		{
+			UEmitterMovieSceneTrack *Track = Cast<UEmitterMovieSceneTrack> (MovieScene->AddMasterTrack(UEmitterMovieSceneTrack::StaticClass()) );
+			 Track->SetEmitter(Emitter);
+		}
+
+	}
+
+
+
+
+	TSharedRef<FTabManager::FLayout> StandaloneDefaultLayout = FTabManager::NewLayout("Standalone_Niagara_Effect_Layout_v7")
 		->AddArea
 		(
-		FTabManager::NewPrimaryArea()->SetOrientation(Orient_Vertical)
-		->Split
-		(
-		FTabManager::NewStack()
-		->SetSizeCoefficient(0.1f)
-		->AddTab(GetToolbarTabId(), ETabState::OpenedTab)
-		->SetHideTabWell(true)
-		)
-		->Split
-		(
-		FTabManager::NewStack()
-		->AddTab(UpdateTabId, ETabState::OpenedTab)
-		)
+			FTabManager::NewPrimaryArea()->SetOrientation(Orient_Vertical)
+			->Split
+			(
+				FTabManager::NewStack()
+				->SetSizeCoefficient(0.1f)
+				->AddTab(GetToolbarTabId(), ETabState::OpenedTab)
+				->SetHideTabWell(true)
+			)
+			->Split
+			(
+				FTabManager::NewSplitter()->SetOrientation(Orient_Horizontal)
+				->Split
+				(
+					FTabManager::NewStack()
+					->SetSizeCoefficient(0.1f)
+					->AddTab(ViewportTabID, ETabState::OpenedTab)
+				)
+				->Split
+				(
+					FTabManager::NewStack()
+					->SetSizeCoefficient(0.3f)
+					->AddTab(EmitterEditorTabID, ETabState::OpenedTab)
+				)
+			)
+			->Split
+			(
+			FTabManager::NewStack()
+			->SetSizeCoefficient(0.3f)
+			->AddTab(CurveEditorTabID, ETabState::OpenedTab)
+			->AddTab(SequencerTabID, ETabState::OpenedTab)
+			)
 		);
 
 	const bool bCreateDefaultStandaloneMenu = true;
@@ -128,7 +202,7 @@ TSharedRef<SNiagaraEffectEditorWidget> FNiagaraEffectEditor::CreateEditorWidget(
 		
 		
 	// Make the effect editor widget
-	return SNew(SNiagaraEffectEditorWidget).TitleBar(TitleBarWidget).EffectObj(InEffect).EffectInstance(EffectInstance);
+	return SNew(SNiagaraEffectEditorWidget).EffectObj(InEffect).EffectInstance(EffectInstance);
 }
 
 
@@ -148,6 +222,76 @@ TSharedRef<SDockTab> FNiagaraEffectEditor::SpawnTab(const FSpawnTabArgs& Args)
 		];
 }
 
+
+
+
+TSharedRef<SDockTab> FNiagaraEffectEditor::SpawnTab_Viewport(const FSpawnTabArgs& Args)
+{
+	check(Args.GetTabId().TabType == ViewportTabID);
+
+	Viewport = SNew(SNiagaraEffectEditorViewport);
+
+	TSharedRef<SDockTab> SpawnedTab =
+		SNew(SDockTab)
+		[
+			Viewport.ToSharedRef()
+		];
+
+	Viewport->SetPreviewEffect(EffectInstance);
+	Viewport->OnAddedToTab(SpawnedTab);
+
+	return SpawnedTab;
+}
+
+
+
+TSharedRef<SDockTab> FNiagaraEffectEditor::SpawnTab_EmitterList(const FSpawnTabArgs& Args)
+{
+	check(Args.GetTabId().TabType == EmitterEditorTabID);
+
+	TSharedRef<SDockTab> SpawnedTab =
+		SNew(SDockTab)
+		[
+			SAssignNew(EmitterEditorWidget, SNiagaraEffectEditorWidget).EffectInstance(EffectInstance).EffectEditor(this)
+		];
+
+
+	return SpawnedTab;
+}
+
+
+TSharedRef<SDockTab> FNiagaraEffectEditor::SpawnTab_CurveEd(const FSpawnTabArgs& Args)
+{
+	check(Args.GetTabId().TabType == CurveEditorTabID);
+
+	TSharedRef<SDockTab> SpawnedTab =
+		SNew(SDockTab)
+		[
+			SAssignNew(TimeLine, SNiagaraTimeline)
+		];
+
+	return SpawnedTab;
+}
+
+
+
+TSharedRef<SDockTab> FNiagaraEffectEditor::SpawnTab_Sequencer(const FSpawnTabArgs& Args)
+{
+	check(Args.GetTabId().TabType == SequencerTabID);
+
+	TSharedRef<SDockTab> SpawnedTab =
+		SNew(SDockTab)
+		[
+			Sequencer->GetSequencerWidget()
+		];
+
+	return SpawnedTab;
+}
+
+
+
+
+
 void FNiagaraEffectEditor::ExtendToolbar()
 {
 	struct Local
@@ -211,8 +355,22 @@ FReply FNiagaraEffectEditor::OnAddEmitterClicked()
 	FNiagaraEmitterProperties *Props = Effect->AddEmitterProperties();
 	TSharedPtr<FNiagaraSimulation> NewEmitter = EffectInstance->AddEmitter(Props);
 	Effect->CreateEffectRendererProps(NewEmitter);
-	UpdateEditorPtr.Pin()->GetViewport()->SetPreviewEffect(EffectInstance);
-	UpdateEditorPtr.Pin()->UpdateList();
+	Viewport->SetPreviewEffect(EffectInstance);
+	EmitterEditorWidget->UpdateList();
+	return FReply::Handled();
+}
+
+
+FReply FNiagaraEffectEditor::OnEmitterSelected(TSharedPtr<FNiagaraSimulation> SelectedItem, ESelectInfo::Type SelType)
+{
+	if (SelectedItem.Get() != nullptr)
+	{
+		if (SelectedItem->GetProperties()->ExternalConstants.GetDataConstants().Num())
+		{
+			FNiagaraDataObject *DataObj = SelectedItem->GetProperties()->ExternalConstants.GetDataConstants().CreateConstIterator().Value();
+			TimeLine.Get()->SetCurve(static_cast<FNiagaraCurveDataObject*>(DataObj)->GetCurveObject());
+		}
+	}
 	return FReply::Handled();
 }
 
diff --git a/Engine/Source/Editor/NiagaraEditor/Private/NiagaraEffectEditor.h b/Engine/Source/Editor/NiagaraEditor/Private/NiagaraEffectEditor.h
--- a/Engine/Source/Editor/NiagaraEditor/Private/NiagaraEffectEditor.h
+++ b/Engine/Source/Editor/NiagaraEditor/Private/NiagaraEffectEditor.h
@@ -1,10 +1,13 @@
 // Copyright 1998-2015 Epic Games, Inc. All Rights Reserved.
 
 #pragma once
-
 #include "INiagaraEffectEditor.h"
 #include "Toolkits/AssetEditorToolkit.h"
 
+#include "SNiagaraEffectEditorViewport.h"
+#include "SNiagaraEffectEditorWidget.h"
+#include "NiagaraSequencer.h"
+
 
 
 /** Viewer/editor for a NiagaraEffect
@@ -30,8 +33,17 @@ public:
 	virtual FLinearColor GetWorldCentricTabColorScale() const override;
 	// End IToolkit interface
 
+	// Delegates for the individual widgets
+	FReply OnEmitterSelected(TSharedPtr<FNiagaraSimulation> SelectedItem, ESelectInfo::Type SelType);
+
 
 	virtual UNiagaraEffect *GetEffect() const	override { return Effect; }	
+
+	static TSharedRef<FMovieSceneTrackEditor> CreateTrackEditor(TSharedRef<ISequencer> InSequencer)
+	{
+		return MakeShareable(new FNiagaraTrackEditor(InSequencer) );
+	}
+
 private:
 	/** Create widget for graph editing */
 	TSharedRef<class SNiagaraEffectEditorWidget> CreateEditorWidget(UNiagaraEffect* InEffect);
@@ -39,15 +51,30 @@ private:
 	/** Spawns the tab with the update graph inside */
 	TSharedRef<SDockTab> SpawnTab(const FSpawnTabArgs& Args);
 
+	TSharedPtr<SNiagaraEffectEditorViewport>	Viewport;
+	TSharedPtr<SNiagaraEffectEditorWidget>	EmitterEditorWidget;
+	TSharedPtr< SNiagaraTimeline > TimeLine;
+
+	TSharedRef<SDockTab> SpawnTab_Viewport(const FSpawnTabArgs& Args);
+	TSharedRef<SDockTab> SpawnTab_EmitterList(const FSpawnTabArgs& Args);
+	TSharedRef<SDockTab> SpawnTab_CurveEd(const FSpawnTabArgs& Args);
+	TSharedRef<SDockTab> SpawnTab_Sequencer(const FSpawnTabArgs& Args);
+
 	/** Builds the toolbar widget */
 	void ExtendToolbar();
 	FReply OnAddEmitterClicked();
+
 private:
 
 	/* The Effect being edited */
 	UNiagaraEffect	*Effect;
 	class FNiagaraEffectInstance *EffectInstance;
 
+	/* stuff needed by the Sequencer */
+	UMovieScene *MovieScene;
+	TSharedPtr<ISequencerObjectBindingManager> SequencerBindingManager;
+	TSharedPtr<ISequencer> Sequencer;
+
 
 	/** The command list for this editor */
 	TSharedPtr<FUICommandList> EditorCommands;
@@ -56,4 +83,8 @@ private:
 
 	/**	Graph editor tab */
 	static const FName UpdateTabId;
+	static const FName ViewportTabID;
+	static const FName EmitterEditorTabID;
+	static const FName CurveEditorTabID;
+	static const FName SequencerTabID;
 };
diff --git a/Engine/Source/Editor/NiagaraEditor/Private/NiagaraSequencer.cpp b/Engine/Source/Editor/NiagaraEditor/Private/NiagaraSequencer.cpp
new file mode 100644
--- /dev/null
+++ b/Engine/Source/Editor/NiagaraEditor/Private/NiagaraSequencer.cpp
@@ -0,0 +1,17 @@
+// Copyright 1998-2015 Epic Games, Inc. All Rights Reserved.
+#pragma once
+#include "NiagaraEditorPrivatePCH.h"
+#include "NiagaraSequencer.h"
+
+
+UEmitterMovieSceneTrack::UEmitterMovieSceneTrack(const FObjectInitializer& ObjectInitializer)
+: Super(ObjectInitializer)
+{
+}
+
+UNiagaraMovieSceneSection::UNiagaraMovieSceneSection(const FObjectInitializer& ObjectInitializer)
+: Super(ObjectInitializer)
+{
+}
+
+ 
\ No newline at end of file
diff --git a/Engine/Source/Editor/NiagaraEditor/Private/NiagaraSequencer.h b/Engine/Source/Editor/NiagaraEditor/Private/NiagaraSequencer.h
new file mode 100644
--- /dev/null
+++ b/Engine/Source/Editor/NiagaraEditor/Private/NiagaraSequencer.h
@@ -0,0 +1,187 @@
+// Copyright 1998-2015 Epic Games, Inc. All Rights Reserved.
+#pragma once
+#include "Editor/Sequencer/Public/ISequencerModule.h"
+#include "Editor/Sequencer/Public/ISequencerObjectBindingManager.h"
+#include "Editor/Sequencer/Public/MovieSceneTrackEditor.h"
+#include "Editor/Sequencer/Public/ISequencerSection.h"
+#include "Runtime/MovieSceneCore/Classes/MovieScene.h"
+#include "Runtime/MovieSceneCore/Classes/MovieSceneTrack.h"
+#include "Runtime/MovieSceneCore/Public/IMovieSceneTrackInstance.h"
+#include "Runtime/MovieSceneCore/Classes/MovieSceneSection.h"
+
+#include "NiagaraSequencer.generated.h"
+
+UCLASS()
+class UNiagaraMovieSceneSection : public UMovieSceneSection
+{
+	GENERATED_UCLASS_BODY()
+public:
+	FText GetEmitterName()	{ return EmitterName; }
+	void SetEmitterName(FText InName)	{ EmitterName = InName; }
+
+private:
+	FText EmitterName;
+};
+
+class INiagaraSequencerSection : public ISequencerSection
+{
+public:
+	INiagaraSequencerSection(UMovieSceneSection &SectionObject)
+	{
+		UNiagaraMovieSceneSection *NiagaraSectionObject = Cast<UNiagaraMovieSceneSection>(&SectionObject);
+		check(NiagaraSectionObject);
+
+		EmitterSection = NiagaraSectionObject;
+	}
+
+	UMovieSceneSection *GetSectionObject(void)
+	{
+		return EmitterSection;
+	}
+
+	int32 OnPaintSection(const FGeometry &AllottedGeometry, const FSlateRect &SectionClippingRect, FSlateWindowElementList &OutDrawElements, int32 LayerId, bool  bParentEnabled) const
+	{
+		FSlateDrawElement::MakeBox
+			(
+			OutDrawElements,
+			LayerId,
+			AllottedGeometry.ToPaintGeometry(),
+			FEditorStyle::GetBrush("CurveEd.TimelineArea"),
+			SectionClippingRect,
+			ESlateDrawEffect::None,
+			FLinearColor(0.3f, 0.3f, 0.6f)
+			);
+
+
+		return 0;
+	}
+
+	FText GetDisplayName(void) const	{ return EmitterSection->GetEmitterName(); }
+	FText GetSectionTitle(void) const	{ return EmitterSection->GetEmitterName(); }
+	float GetSectionHeight() const		{ return 32.0f; }
+
+	void ISequencerSection::GenerateSectionLayout(ISectionLayoutBuilder &) const {}
+
+private:
+	UNiagaraMovieSceneSection *EmitterSection;
+};
+
+
+
+class INiagaraTrackInstance : public IMovieSceneTrackInstance
+{
+public:
+	INiagaraTrackInstance(TSharedPtr<FNiagaraSimulation> InEmitter)
+	{
+		Emitter = InEmitter;
+	}
+
+	virtual void Update(float Position, float LastPosition, const TArray<UObject*>& RuntimeObjects, class IMovieScenePlayer& Player)
+	{}
+	virtual void RefreshInstance(const TArray<UObject*>& RuntimeObjects, class IMovieScenePlayer& Player)
+	{}
+
+private:
+	TSharedPtr<FNiagaraSimulation> Emitter;
+};
+
+
+
+
+
+
+
+UCLASS(MinimalAPI)
+class UEmitterMovieSceneTrack : public UMovieSceneTrack
+{
+	GENERATED_UCLASS_BODY()
+public:
+
+	/** UMovieSceneTrack interface */
+	virtual FName GetTrackName() const override { return *Emitter->GetProperties()->Name; }
+	void SetEmitter(TSharedPtr<FNiagaraSimulation> InEmitter)
+	{
+		Emitter = InEmitter;
+		const FNiagaraEmitterProperties *EmitterProps = InEmitter->GetProperties();
+
+		UNiagaraMovieSceneSection *Section = NewObject<UNiagaraMovieSceneSection>(this, *Emitter->GetProperties()->Name);
+		Section->SetStartTime(EmitterProps->StartTime);
+		Section->SetEndTime(EmitterProps->EndTime);
+		Section->SetEmitterName(FText::FromString(EmitterProps->Name));
+
+		Sections.Add(Section);
+	}
+
+	virtual void RemoveAllAnimationData() override	{};
+	virtual bool HasSection(UMovieSceneSection* Section) const override { return false; }
+	virtual bool IsEmpty() const override { return false; }
+	virtual TSharedPtr<class IMovieSceneTrackInstance> CreateInstance()
+	{
+		return MakeShareable(new INiagaraTrackInstance(Emitter));
+	}
+	virtual TArray<UMovieSceneSection*> GetAllSections() const
+	{
+		return Sections;
+	}
+private:
+	TSharedPtr<FNiagaraSimulation> Emitter;
+	TArray<UMovieSceneSection*> Sections;
+};
+
+
+
+class FNiagaraTrackEditor : public FMovieSceneTrackEditor
+{
+public:
+	FNiagaraTrackEditor(TSharedPtr<ISequencer> Sequencer) : FMovieSceneTrackEditor(Sequencer.ToSharedRef())
+	{
+
+	}
+
+	virtual bool SupportsType(TSubclassOf<class UMovieSceneTrack> TrackClass) const
+	{
+		if (TrackClass == UEmitterMovieSceneTrack::StaticClass())
+		{
+			return true;
+		}
+		return false;
+	}
+
+	virtual TSharedRef<ISequencerSection> MakeSectionInterface(class UMovieSceneSection& SectionObject, UMovieSceneTrack* Track)
+	{
+		INiagaraSequencerSection *Sec = new INiagaraSequencerSection(SectionObject);
+		return MakeShareable(Sec);
+	}
+
+};
+
+
+
+
+/** Sequencer object binding manager for Niagara effects;
+*/
+class FNiagaraSequencerObjectBindingManager : public FGCObject, public ISequencerObjectBindingManager
+{
+public:
+	FNiagaraSequencerObjectBindingManager()
+	{}
+
+	~FNiagaraSequencerObjectBindingManager()
+	{}
+
+	/** FGCObject interface */
+	virtual void AddReferencedObjects(FReferenceCollector& Collector) override {	}
+
+	/** ISequencerObjectBindingManager interface */
+	virtual bool AllowsSpawnableObjects() const override { return false; }
+	virtual FGuid FindGuidForObject(const UMovieScene& MovieScene, UObject& Object) const override { return FGuid(); }
+	virtual void SpawnOrDestroyObjectsForInstance(TSharedRef<FMovieSceneInstance> MovieSceneInstance, bool bDestroyAll) override { }
+	virtual void DestroyAllSpawnedObjects() override { }
+	virtual bool CanPossessObject(UObject& Object) const override 	{ return false; }
+	virtual void BindPossessableObject(const FGuid& PossessableGuid, UObject& PossessedObject) override	{};
+	virtual void UnbindPossessableObjects(const FGuid& PossessableGuid) override	{};
+	virtual void GetRuntimeObjects(const TSharedRef<FMovieSceneInstance>& MovieSceneInstance, const FGuid& ObjectGuid, TArray<UObject*>& OutRuntimeObjects) const override	{}
+
+	virtual bool TryGetObjectBindingDisplayName(const FGuid& ObjectGuid, FText& DisplayName) const override	{ return false; }
+private:
+};
\ No newline at end of file
diff --git a/Engine/Source/Editor/NiagaraEditor/Private/SNiagaraEffectEditorViewport.h b/Engine/Source/Editor/NiagaraEditor/Private/SNiagaraEffectEditorViewport.h
--- a/Engine/Source/Editor/NiagaraEditor/Private/SNiagaraEffectEditorViewport.h
+++ b/Engine/Source/Editor/NiagaraEditor/Private/SNiagaraEffectEditorViewport.h
@@ -4,7 +4,7 @@
 
 #include "PreviewScene.h"
 #include "SEditorViewport.h"
-
+#include "NiagaraEffect.h"
 
 /**
  * Material Editor Preview viewport widget
diff --git a/Engine/Source/Editor/NiagaraEditor/Private/SNiagaraEffectEditorWidget.cpp b/Engine/Source/Editor/NiagaraEditor/Private/SNiagaraEffectEditorWidget.cpp
--- a/Engine/Source/Editor/NiagaraEditor/Private/SNiagaraEffectEditorWidget.cpp
+++ b/Engine/Source/Editor/NiagaraEditor/Private/SNiagaraEffectEditorWidget.cpp
@@ -7,6 +7,7 @@
 #include "SSplitter.h"
 #include "Editor/PropertyEditor/Public/PropertyEditorModule.h"
 #include "Editor/PropertyEditor/Public/IDetailsView.h"
+#include "NiagaraEffectEditor.h"
 
 
 #define LOCTEXT_NAMESPACE "NiagaraEffectEditor"
@@ -25,111 +26,42 @@ void SNiagaraEffectEditorWidget::Construct(const FArguments& InArgs)
 		EffectInstance = new FNiagaraEffectInstance(EffectObj);
 	}
 
-	//bNeedsRefresh = false;
-
-
-	Viewport = SNew(SNiagaraEffectEditorViewport);
-
+	EffectEditor = InArgs._EffectEditor;
 
 	TSharedPtr<SOverlay> OverlayWidget;
 	this->ChildSlot
 		[
+
 			SNew(SSplitter)
 			.Orientation(Orient_Vertical)
 			+ SSplitter::Slot()
 			.Value(0.75f)
 			[
 
-			SNew(SVerticalBox)
-			+ SVerticalBox::Slot()
-			.AutoHeight()
-			[
-
-				SNew(SHorizontalBox)
-				+ SHorizontalBox::Slot()
-				.AutoWidth()
+				SNew(SVerticalBox)
+				+ SVerticalBox::Slot()
+				.AutoHeight()
 				[
-					SNew(SVerticalBox)
-					+ SVerticalBox::Slot()
-					.AutoHeight()
-					[
-						SNew(SBox)
-						.WidthOverride(512)
-						.HeightOverride(512)
-						[
-							Viewport.ToSharedRef()
-						]
-					]
-				]
-				+ SHorizontalBox::Slot()
-					.AutoWidth()
-					[
-						SAssignNew(OverlayWidget, SOverlay)
-
-						+SOverlay::Slot()
-						.Padding(10)
-						.VAlign(VAlign_Fill)
+					SNew(SHorizontalBox)
+					+ SHorizontalBox::Slot()
+						.AutoWidth()
 						[
-							SAssignNew(ListView, SListView<TSharedPtr<FNiagaraSimulation> >)
-							.ItemHeight(256)
-							.ListItemsSource(&(this->EffectInstance->GetEmitters()))
-							.OnGenerateRow(this, &SNiagaraEffectEditorWidget::OnGenerateWidgetForList)
-							.OnSelectionChanged(this, &SNiagaraEffectEditorWidget::OnSelectionChanged)
-							//.SelectionMode(ESelectionMode::None)
-						]
-
-						// Bottom-right corner text indicating the type of tool
-						+ SOverlay::Slot()
+							SAssignNew(OverlayWidget, SOverlay)
+							+SOverlay::Slot()
 							.Padding(10)
-							.VAlign(VAlign_Bottom)
-							.HAlign(HAlign_Right)
+							.VAlign(VAlign_Fill)
 							[
-								SNew(STextBlock)
-								.Visibility(EVisibility::HitTestInvisible)
-								.TextStyle(FEditorStyle::Get(), "Graph.CornerText")
-								.Text(LOCTEXT("NiagaraEditorLabel", "Niagara Effect"))
+								SAssignNew(ListView, SListView<TSharedPtr<FNiagaraSimulation> >)
+								.ItemHeight(256)
+								.ListItemsSource(&(this->EffectInstance->GetEmitters()))
+								.OnGenerateRow(this, &SNiagaraEffectEditorWidget::OnGenerateWidgetForList)
+								.OnSelectionChanged(this, &SNiagaraEffectEditorWidget::OnSelectionChanged)
+								//.SelectionMode(ESelectionMode::Single)
 							]
-
-						// Top-right corner text indicating read only when not simulating
-						+ SOverlay::Slot()
-							.Padding(20)
-							.VAlign(VAlign_Top)
-							.HAlign(HAlign_Right)
-							[
-								SNew(STextBlock)
-								.Visibility(this, &SNiagaraEffectEditorWidget::ReadOnlyVisibility)
-								.TextStyle(FEditorStyle::Get(), "Graph.CornerText")
-								.Text(LOCTEXT("ReadOnlyLabel", "Read Only"))
-							]
-
-						// Bottom-right corner text for notification list position
-						+ SOverlay::Slot()
-							.Padding(15.f)
-							.VAlign(VAlign_Bottom)
-							.HAlign(HAlign_Right)
-							[
-								SAssignNew(NotificationListPtr, SNotificationList)
-								.Visibility(EVisibility::HitTestInvisible)
-							]
-					]
-			]
-
-			]
-			+ SSplitter::Slot()
-			.Value(0.25f)
-				[
-					NGED_SECTION_BORDER
-					[
-						SNew(SBox)
-						.MinDesiredHeight(400).MinDesiredWidth(1920)
-						[
-							SAssignNew(TimeLine, SNiagaraTimeline)
 						]
-					]
 				]
+			]
 		];
-
-	Viewport->SetPreviewEffect(EffectInstance);
 }
 
 
@@ -589,4 +521,10 @@ void SEmitterWidget::OnSpawnScriptSelectedFromPicker(UObject *Asset)
 }
 
 
+void SNiagaraEffectEditorWidget::OnSelectionChanged(TSharedPtr<FNiagaraSimulation> SelectedItem, ESelectInfo::Type SelType)
+{
+	EffectEditor->OnEmitterSelected(SelectedItem, SelType);
+}
+
+
 END_SLATE_FUNCTION_BUILD_OPTIMIZATION
\ No newline at end of file
diff --git a/Engine/Source/Editor/NiagaraEditor/Private/SNiagaraEffectEditorWidget.h b/Engine/Source/Editor/NiagaraEditor/Private/SNiagaraEffectEditorWidget.h
--- a/Engine/Source/Editor/NiagaraEditor/Private/SNiagaraEffectEditorWidget.h
+++ b/Engine/Source/Editor/NiagaraEditor/Private/SNiagaraEffectEditorWidget.h
@@ -19,7 +19,6 @@
 #include "SNumericEntryBox.h"
 #include "Curves/CurveVector.h"
 
-
 #define NGED_SECTION_BORDER SNew(SBorder).BorderImage(FEditorStyle::GetBrush("ToolPanel.GroupBorder")).Padding(2.0f).HAlign(HAlign_Left)
 #define NGED_SECTION_LIGHTBORDER SNew(SBorder).BorderImage(FEditorStyle::GetBrush("ToolPanel.LightGroupBorder")).Padding(2.0f).HAlign(HAlign_Left)
 #define NGED_SECTION_DARKBORDER SNew(SBorder).BorderImage(FEditorStyle::GetBrush("ToolPanel.DarkGroupBorder")).Padding(2.0f).HAlign(HAlign_Left)
@@ -430,11 +429,10 @@ private:
 	/** Notification list to pass messages to editor users  */
 	TSharedPtr<SNotificationList> NotificationListPtr;
 	SOverlay::FOverlaySlot* DetailsViewSlot;
-	TSharedPtr<SNiagaraEffectEditorViewport>	Viewport;
 	
 	TSharedPtr< SListView<TSharedPtr<FNiagaraSimulation> > > ListView;
 
-	TSharedPtr< SNiagaraTimeline > TimeLine;
+	class FNiagaraEffectEditor *EffectEditor;
 public:
 	SLATE_BEGIN_ARGS(SNiagaraEffectEditorWidget)
 		: _EffectObj(nullptr)
@@ -443,6 +441,7 @@ public:
 	SLATE_ARGUMENT(UNiagaraEffect*, EffectObj)
 	SLATE_ARGUMENT(TSharedPtr<SWidget>, TitleBar)
 	SLATE_ARGUMENT(FNiagaraEffectInstance*, EffectInstance)
+	SLATE_ARGUMENT(class FNiagaraEffectEditor*, EffectEditor)
 	SLATE_END_ARGS()
 
 
@@ -457,8 +456,6 @@ public:
 		ListView->RequestListRefresh();
 	}
 
-	TSharedPtr<SNiagaraEffectEditorViewport> GetViewport()	{ return Viewport; }
-
 	TSharedRef<ITableRow> OnGenerateWidgetForList(TSharedPtr<FNiagaraSimulation> InItem, const TSharedRef< STableViewBase >& OwnerTable)
 	{
 		return SNew(STableRow<TSharedPtr<FString>>, OwnerTable)
@@ -473,17 +470,7 @@ public:
 			];
 	}
 
-	void OnSelectionChanged( TSharedPtr<FNiagaraSimulation> SelectedItem, ESelectInfo::Type SelType)
-	{
-		if ( SelectedItem.Get() != nullptr )
-		{
-			if (SelectedItem->GetProperties()->ExternalConstants.GetDataConstants().Num())
-			{
-				FNiagaraDataObject *DataObj = SelectedItem->GetProperties()->ExternalConstants.GetDataConstants().CreateConstIterator().Value();
-				TimeLine->SetCurve(static_cast<FNiagaraCurveDataObject*>(DataObj)->GetCurveObject());
-			}
-		}
-	}
+	void OnSelectionChanged(TSharedPtr<FNiagaraSimulation> SelectedItem, ESelectInfo::Type SelType); 
 
 	void Construct(const FArguments& InArgs);
 
diff --git a/Engine/Source/Runtime/Niagara/Classes/NiagaraSimulation.h b/Engine/Source/Runtime/Niagara/Classes/NiagaraSimulation.h
--- a/Engine/Source/Runtime/Niagara/Classes/NiagaraSimulation.h
+++ b/Engine/Source/Runtime/Niagara/Classes/NiagaraSimulation.h
@@ -40,7 +40,7 @@ struct FNiagaraEmitterProperties
 public:
 	FNiagaraEmitterProperties() : 
 		bIsEnabled(true), SpawnRate(50), UpdateScript(nullptr), SpawnScript(nullptr), Material(nullptr), 
-		RenderModuleType(RMT_Sprites), RendererProperties(nullptr)
+		RenderModuleType(RMT_Sprites), RendererProperties(nullptr), StartTime(0.0f), EndTime(0.0f)
 	{
 		Name = FString(TEXT("New Emitter"));
 	}
@@ -59,6 +59,12 @@ public:
 	UPROPERTY(EditAnywhere, Category = "Emitter Properties")
 	TEnumAsByte<EEmitterRenderModuleType> RenderModuleType;
 
+	UPROPERTY(EditAnywhere, Category = "Emitter Properties")
+	float StartTime;
+	UPROPERTY(EditAnywhere, Category = "Emitter Properties")
+	float EndTime;
+
+
 	UPROPERTY()
 	class UNiagaraEffectRendererProperties *RendererProperties;
 
