Fix for compile error.

diff --git a/Engine/Source/Editor/NiagaraEditor/Private/NiagaraSequencer.h b/Engine/Source/Editor/NiagaraEditor/Private/NiagaraSequencer.h
--- a/Engine/Source/Editor/NiagaraEditor/Private/NiagaraSequencer.h
+++ b/Engine/Source/Editor/NiagaraEditor/Private/NiagaraSequencer.h
@@ -60,7 +60,7 @@ public:
 	FText GetSectionTitle(void) const	{ return EmitterSection->GetEmitterName(); }
 	float GetSectionHeight() const		{ return 32.0f; }
 
-	void ISequencerSection::GenerateSectionLayout(ISectionLayoutBuilder &) const {}
+	void GenerateSectionLayout(ISectionLayoutBuilder &) const {}
 
 private:
 	UNiagaraMovieSceneSection *EmitterSection;
@@ -77,7 +77,8 @@ public:
 	}
 
 	virtual void Update(float Position, float LastPosition, const TArray<UObject*>& RuntimeObjects, class IMovieScenePlayer& Player)
-	{}
+	{
+	}
 	virtual void RefreshInstance(const TArray<UObject*>& RuntimeObjects, class IMovieScenePlayer& Player)
 	{}
 
