Remove pragma once from NiagaraSequencer.cpp to fix non unity on mac CIS error.

#UE4

diff --git a/Engine/Source/Editor/NiagaraEditor/Private/NiagaraSequencer.cpp b/Engine/Source/Editor/NiagaraEditor/Private/NiagaraSequencer.cpp
--- a/Engine/Source/Editor/NiagaraEditor/Private/NiagaraSequencer.cpp
+++ b/Engine/Source/Editor/NiagaraEditor/Private/NiagaraSequencer.cpp
@@ -1,5 +1,4 @@
 // Copyright 1998-2015 Epic Games, Inc. All Rights Reserved.
-#pragma once
 #include "NiagaraEditorPrivatePCH.h"
 #include "NiagaraSequencer.h"
 
