Niagara: embedded sequencer timeline in the effect editor, with lifetime editing for emitters; basic lifetime management for emitters in the backend;

diff --git a/Engine/Source/Editor/NiagaraEditor/Private/NiagaraSequencer.cpp b/Engine/Source/Editor/NiagaraEditor/Private/NiagaraSequencer.cpp
--- a/Engine/Source/Editor/NiagaraEditor/Private/NiagaraSequencer.cpp
+++ b/Engine/Source/Editor/NiagaraEditor/Private/NiagaraSequencer.cpp
@@ -13,4 +13,19 @@ UNiagaraMovieSceneSection::UNiagaraMovieSceneSection(const FObjectInitializer& O
 {
 }
 
- 
\ No newline at end of file
+
+/*
+ *  This is called when the user edits a track in the effect editor timeline
+ */
+void INiagaraTrackInstance::RefreshInstance(const TArray<UObject*>& RuntimeObjects, class IMovieScenePlayer& Player)
+{
+	// every track should have exactly one section, always
+	check(Track->GetAllSections().Num() == 1)
+
+	UNiagaraMovieSceneSection *Section = Cast<UNiagaraMovieSceneSection>( Track->GetAllSections()[0] );
+	
+	check(Section);
+
+	Track->GetEmitter()->GetProperties()->StartTime = Section->GetStartTime();
+	Track->GetEmitter()->GetProperties()->EndTime = Section->GetEndTime();
+}
diff --git a/Engine/Source/Editor/NiagaraEditor/Private/NiagaraSequencer.h b/Engine/Source/Editor/NiagaraEditor/Private/NiagaraSequencer.h
--- a/Engine/Source/Editor/NiagaraEditor/Private/NiagaraSequencer.h
+++ b/Engine/Source/Editor/NiagaraEditor/Private/NiagaraSequencer.h
@@ -1,4 +1,5 @@
 // Copyright 1998-2015 Epic Games, Inc. All Rights Reserved.
+
 #pragma once
 #include "Editor/Sequencer/Public/ISequencerModule.h"
 #include "Editor/Sequencer/Public/ISequencerObjectBindingManager.h"
@@ -11,6 +12,9 @@
 
 #include "NiagaraSequencer.generated.h"
 
+/** 
+ *	Niagara editor movie scene section; represents one emitter in the timeline
+ */
 UCLASS()
 class UNiagaraMovieSceneSection : public UMovieSceneSection
 {
@@ -18,11 +22,17 @@ class UNiagaraMovieSceneSection : public UMovieSceneSection
 public:
 	FText GetEmitterName()	{ return EmitterName; }
 	void SetEmitterName(FText InName)	{ EmitterName = InName; }
+	void SetEmitterProps(const FNiagaraEmitterProperties *Props)	{ EmitterProps = Props; }
 
+	const FNiagaraEmitterProperties *GetEmitterProps() const	{ return EmitterProps; }
 private:
 	FText EmitterName;
+	const FNiagaraEmitterProperties *EmitterProps;
 };
 
+/**
+*	Visual representation of UNiagaraMovieSceneSection
+*/
 class INiagaraSequencerSection : public ISequencerSection
 {
 public:
@@ -41,6 +51,7 @@ public:
 
 	int32 OnPaintSection(const FGeometry &AllottedGeometry, const FSlateRect &SectionClippingRect, FSlateWindowElementList &OutDrawElements, int32 LayerId, bool  bParentEnabled) const
 	{
+		// draw the first run of the emitter
 		FSlateDrawElement::MakeBox
 			(
 			OutDrawElements,
@@ -53,6 +64,29 @@ public:
 			);
 
 
+		// draw all loops of the emitter as 'ghosts' of the original section
+		float X = AllottedGeometry.AbsolutePosition.X;
+		float GeomW = AllottedGeometry.GetDrawSize().X;
+		float ClipW = SectionClippingRect.Right-SectionClippingRect.Left;
+		FSlateRect NewClipRect = SectionClippingRect;
+		NewClipRect.Left += 5;
+		NewClipRect.Right -= 10;
+		for (int32 Loop = 0; Loop < EmitterSection->GetEmitterProps()->NumLoops - 1; Loop++)
+		{
+			NewClipRect.Left += (GeomW) - FMath::Max(0.0f, GeomW - ClipW);
+			NewClipRect.Right += GeomW;
+			FSlateDrawElement::MakeBox
+				(
+				OutDrawElements,
+				LayerId,
+				AllottedGeometry.ToPaintGeometry(FVector2D(GeomW*(Loop+1), 0.0f), AllottedGeometry.GetDrawSize(), 1.0f),
+				FEditorStyle::GetBrush("CurveEd.TimelineArea"),
+				NewClipRect,
+				ESlateDrawEffect::None,
+				FLinearColor(0.3f, 0.3f, 0.6f, 0.25f)
+				);
+		}
+
 		return 0;
 	}
 
@@ -68,30 +102,36 @@ private:
 
 
 
+
+
+/**
+*	One instance of a UEmitterMovieSceneTrack
+*/
 class INiagaraTrackInstance : public IMovieSceneTrackInstance
 {
 public:
-	INiagaraTrackInstance(TSharedPtr<FNiagaraSimulation> InEmitter)
+	INiagaraTrackInstance(class UEmitterMovieSceneTrack *InTrack)
+		: Track(InTrack)
 	{
-		Emitter = InEmitter;
 	}
 
 	virtual void Update(float Position, float LastPosition, const TArray<UObject*>& RuntimeObjects, class IMovieScenePlayer& Player)
 	{
 	}
-	virtual void RefreshInstance(const TArray<UObject*>& RuntimeObjects, class IMovieScenePlayer& Player)
-	{}
 
+	virtual void RefreshInstance(const TArray<UObject*>& RuntimeObjects, class IMovieScenePlayer& Player);
 private:
 	TSharedPtr<FNiagaraSimulation> Emitter;
+	class UEmitterMovieSceneTrack *Track;
 };
 
 
 
 
 
-
-
+/**
+*	Single track containing exactly one UNiagaraMovieSceneSection, representing one emitter
+*/
 UCLASS(MinimalAPI)
 class UEmitterMovieSceneTrack : public UMovieSceneTrack
 {
@@ -106,6 +146,8 @@ public:
 		const FNiagaraEmitterProperties *EmitterProps = InEmitter->GetProperties();
 
 		UNiagaraMovieSceneSection *Section = NewObject<UNiagaraMovieSceneSection>(this, *Emitter->GetProperties()->Name);
+		Section->SetEmitterProps(EmitterProps);
+
 		Section->SetStartTime(EmitterProps->StartTime);
 		Section->SetEndTime(EmitterProps->EndTime);
 		Section->SetEmitterName(FText::FromString(EmitterProps->Name));
@@ -118,12 +160,15 @@ public:
 	virtual bool IsEmpty() const override { return false; }
 	virtual TSharedPtr<class IMovieSceneTrackInstance> CreateInstance()
 	{
-		return MakeShareable(new INiagaraTrackInstance(Emitter));
+		return MakeShareable(new INiagaraTrackInstance(this));
 	}
 	virtual TArray<UMovieSceneSection*> GetAllSections() const
 	{
 		return Sections;
 	}
+
+	TSharedPtr<FNiagaraSimulation> GetEmitter() { return Emitter; }
+
 private:
 	TSharedPtr<FNiagaraSimulation> Emitter;
 	TArray<UMovieSceneSection*> Sections;
@@ -131,6 +176,9 @@ private:
 
 
 
+/**
+*	Track editor for Niagara emitter tracks
+*/
 class FNiagaraTrackEditor : public FMovieSceneTrackEditor
 {
 public:
@@ -158,8 +206,8 @@ public:
 
 
 
-
-/** Sequencer object binding manager for Niagara effects;
+/**
+ * Sequencer object binding manager for Niagara effects; all of this is just stubbed out
 */
 class FNiagaraSequencerObjectBindingManager : public FGCObject, public ISequencerObjectBindingManager
 {
diff --git a/Engine/Source/Editor/NiagaraEditor/Private/SNiagaraEffectEditorWidget.cpp b/Engine/Source/Editor/NiagaraEditor/Private/SNiagaraEffectEditorWidget.cpp
--- a/Engine/Source/Editor/NiagaraEditor/Private/SNiagaraEffectEditorWidget.cpp
+++ b/Engine/Source/Editor/NiagaraEditor/Private/SNiagaraEffectEditorWidget.cpp
@@ -306,7 +306,7 @@ void SEmitterWidget::Construct(const FArguments& InArgs)
 						]
 
 
-						// Spawn rate
+						// Spawn rate and other emitter properties
 						+SHorizontalBox::Slot()
 							.AutoWidth()
 							.HAlign(HAlign_Left)
@@ -335,6 +335,27 @@ void SEmitterWidget::Construct(const FArguments& InArgs)
 												.Text(this, &SEmitterWidget::GetSpawnRateText)
 											]
 									]
+									+ SVerticalBox::Slot()
+									.Padding(4)
+									.AutoHeight()
+									[
+										SNew(SHorizontalBox)
+										+ SHorizontalBox::Slot()
+										.HAlign(HAlign_Center)
+										.VAlign(VAlign_Center)
+										.AutoWidth()
+										.Padding(4)
+										[
+											SNew(STextBlock)
+											.Text(LOCTEXT("LoopsLabel", "Loops"))
+										]
+										+ SHorizontalBox::Slot()
+											[
+												SNew(SEditableTextBox).OnTextChanged(this, &SEmitterWidget::OnNumLoopsChanged)
+												.Text(this, &SEmitterWidget::GetLoopsText)
+											]
+									]
+
 								]
 							]
 
diff --git a/Engine/Source/Editor/NiagaraEditor/Private/SNiagaraEffectEditorWidget.h b/Engine/Source/Editor/NiagaraEditor/Private/SNiagaraEffectEditorWidget.h
--- a/Engine/Source/Editor/NiagaraEditor/Private/SNiagaraEffectEditorWidget.h
+++ b/Engine/Source/Editor/NiagaraEditor/Private/SNiagaraEffectEditorWidget.h
@@ -33,9 +33,6 @@ class SVectorConstantWidget : public SCompoundWidget, public FNotifyHook
 private:	
 	FName ConstantName;
 	TSharedPtr<FNiagaraSimulation> Emitter;
-
-	//SNumericEntryBox<float> XText, YText, ZText, WText;
-
 public:
 	SLATE_BEGIN_ARGS(SVectorConstantWidget) :
 	_Emitter(nullptr),
@@ -352,12 +349,24 @@ public:
 	FText GetSpawnRateText() const
 	{
 		TCHAR Txt[32];
-		//FCString::Snprintf(Txt, 32, TEXT("%f"), Emitter->GetSpawnRate() );
 		FCString::Snprintf(Txt, 32, TEXT("%f"), Emitter->GetProperties()->SpawnRate );
+		return FText::FromString(Txt);
+	}
+
+	void OnNumLoopsChanged(const FText &InText)
+	{
+		int Loops = FCString::Atoi(*InText.ToString());
+		Emitter->GetProperties()->NumLoops = Loops;
+	}
 
+	FText GetLoopsText() const
+	{
+		TCHAR Txt[32];
+		FCString::Snprintf(Txt, 32, TEXT("%i"), Emitter->GetProperties()->NumLoops);
 		return FText::FromString(Txt);
 	}
 
+
 	FText GetRenderModuleText() const
 	{
 		if (Emitter->GetProperties()->RenderModuleType == RMT_None)
diff --git a/Engine/Source/Runtime/Engine/Classes/Engine/NiagaraEffectRenderer.h b/Engine/Source/Runtime/Engine/Classes/Engine/NiagaraEffectRenderer.h
--- a/Engine/Source/Runtime/Engine/Classes/Engine/NiagaraEffectRenderer.h
+++ b/Engine/Source/Runtime/Engine/Classes/Engine/NiagaraEffectRenderer.h
@@ -17,6 +17,7 @@ NiagaraEffectRenderer.h: Base class for Niagara render modules
 DECLARE_CYCLE_STAT(TEXT("Generate Sprite Vertex Data"), STAT_NiagaraGenSpriteVertexData, STATGROUP_Niagara);
 DECLARE_CYCLE_STAT(TEXT("Generate Ribbon Vertex Data"), STAT_NiagaraGenRibbonVertexData, STATGROUP_Niagara);
 
+
 /** Struct used to pass dynamic data from game thread to render thread */
 struct FNiagaraDynamicDataBase
 {
@@ -102,19 +103,21 @@ public:
 
 	float GetCPUTimeMS() { return CPUTimeMS; }
 	FBox GetBounds()	{ return CachedBounds;  }
-protected:
-	mutable float CPUTimeMS;
 
+protected:
 	NiagaraEffectRenderer()	
-		: Material(nullptr)
+		: Material(nullptr), TickState(NTS_Running)
+
 	{
 	}
 
+	mutable float CPUTimeMS;
 	UMaterial* Material;
 	FBox CachedBounds;
 
 private:
 	FMaterialRelevance MaterialRelevance;
+	ENiagaraTickState TickState;
 };
 
 
diff --git a/Engine/Source/Runtime/Niagara/Classes/NiagaraEffect.h b/Engine/Source/Runtime/Niagara/Classes/NiagaraEffect.h
--- a/Engine/Source/Runtime/Niagara/Classes/NiagaraEffect.h
+++ b/Engine/Source/Runtime/Niagara/Classes/NiagaraEffect.h
@@ -8,6 +8,14 @@
 #include "NiagaraEffect.generated.h"
 
 
+enum NiagaraTickState
+{
+	ERunning,		// normally running
+	ESuspended,		// stop simulating and spawning, still render
+	EDieing,		// stop spawning, still simulate and render
+	EDead			// no live particles, no new spawning
+};
+
 UCLASS()
 class NIAGARA_API UNiagaraEffect : public UObject
 {
@@ -90,6 +98,8 @@ public:
 		Component = InComponent;
 		InitRenderModules(Component->GetWorld()->FeatureLevel);
 		RenderModuleupdate();
+
+		Age = 0.0f;
 	}
 
 	NIAGARA_API TSharedPtr<FNiagaraSimulation> AddEmitter(FNiagaraEmitterProperties *Properties);
@@ -126,12 +136,44 @@ public:
 
 		for (TSharedPtr<FNiagaraSimulation>&it : Emitters)
 		{
-			it->SetConstants(Constants);
-			it->GetConstants().Merge(it->GetProperties()->ExternalConstants);
-			it->Tick(DeltaSeconds);
+			FNiagaraEmitterProperties *Props = it->GetProperties();
+
+			int Duration = Props->EndTime - Props->StartTime;
+			int LoopedStartTime = Props->StartTime + Duration*it->GetLoopCount();
+			int LoopedEndTime = Props->EndTime + Duration*it->GetLoopCount();
+
+			// manage emitter lifetime
+			//
+			if (Props->StartTime == 0.0f && Props->EndTime == 0.0f
+				|| (LoopedStartTime<Age && LoopedEndTime>Age))
+			{
+				it->SetTickState(NTS_Running);
+			}
+			else
+			{
+				// if we're past end time, manage looping; we reset the emitters age constant
+				// if it has one
+				if (Props->NumLoops > 1 && it->GetLoopCount()<Props->NumLoops)
+				{
+					it->LoopRestart();
+				}
+				else
+				{
+					it->SetTickState(NTS_Dieing);
+				}
+			}
+
+			if (it->GetTickState() != NTS_Dead && it->GetTickState() != NTS_Suspended)
+			{
+				it->SetConstants(Constants);
+				it->GetConstants().Merge(it->GetProperties()->ExternalConstants);
+				it->Tick(DeltaSeconds);
+			}
 
 			EffectBounds += it->GetEffectRenderer()->GetBounds();
 		}
+
+		Age += DeltaSeconds;
 	}
 
 	NIAGARA_API void RenderModuleupdate();
@@ -150,6 +192,7 @@ public:
 private:
 	UNiagaraComponent *Component;
 	FBox EffectBounds;
+	float Age;
 
 	/** Local constant table. */
 	FNiagaraConstantMap Constants;
diff --git a/Engine/Source/Runtime/Niagara/Classes/NiagaraSimulation.h b/Engine/Source/Runtime/Niagara/Classes/NiagaraSimulation.h
--- a/Engine/Source/Runtime/Niagara/Classes/NiagaraSimulation.h
+++ b/Engine/Source/Runtime/Niagara/Classes/NiagaraSimulation.h
@@ -29,6 +29,7 @@ enum EEmitterRenderModuleType
 
 
 
+
 /** 
  *	FniagaraEmitterProperties stores the attributes of an FNiagaraSimulation
  *	that need to be serialized and are used for its initialization 
@@ -47,7 +48,8 @@ public:
 		RenderModuleType(RMT_Sprites),
 		StartTime(0.0f),
 		EndTime(0.0f),
-		RendererProperties(nullptr)
+		RendererProperties(nullptr),
+		NumLoops(0)
 	{
 		Name = FString(TEXT("New Emitter"));
 	}
@@ -80,6 +82,9 @@ public:
 
 	UPROPERTY(EditAnywhere, Category = "Emitter Properties")
 	FNiagaraConstantMap ExternalSpawnConstants;		// these are the spawn script constants from the effect editor; will be added to the emitter's constant map
+
+	UPROPERTY(EditAnywhere, Category = "Emitter Properties")
+	int32 NumLoops;
 };
 
 
@@ -125,15 +130,31 @@ public:
 	void NIAGARA_API SetSpawnScript(UNiagaraScript* Script);
 	void NIAGARA_API SetUpdateScript(UNiagaraScript* Script);
 	void NIAGARA_API SetEnabled(bool bEnabled);
+
+	float GetAge() { return Age; }
+	int32 GetLoopCount()	{ return Loops; }
+	void LoopRestart()	
+	{ 
+		Age = 0.0f;
+		Loops++;
+		SetTickState(NTS_Running);
+	}
+
+	ENiagaraTickState NIAGARA_API GetTickState()	{ return TickState; }
+	void NIAGARA_API SetTickState(ENiagaraTickState InState)	{ TickState = InState; }
 private:
 	FNiagaraEmitterProperties *Props;		// points to an entry in the array of FNiagaraProperties stored in the EffectInstance (itself pointing to the effect's properties)
 
 	/* The age of the emitter*/
 	float Age;
+	/* how many loops this emitter has gone through */
+	uint32 Loops;
 	/* If false, don't tick or render*/
 	bool bIsEnabled;
 	/* Seconds taken to process everything (including rendering) */
 	float CPUTimeMS;
+	/* Emitter tick state */
+	ENiagaraTickState TickState;
 	
 	/** Local constant set. */
 	FNiagaraConstantMap Constants;
@@ -151,6 +172,11 @@ private:
 	/** Calc number to spawn */
 	int32 CalcNumToSpawn(float DeltaSeconds)
 	{
+		if (TickState == NTS_Dead || TickState == NTS_Dieing || TickState == NTS_Suspended)
+		{
+			return 0;
+		}
+
 		float FloatNumToSpawn = SpawnRemainder + (DeltaSeconds * Props->SpawnRate);
 		int32 NumToSpawn = FMath::FloorToInt(FloatNumToSpawn);
 		SpawnRemainder = FloatNumToSpawn - NumToSpawn;
diff --git a/Engine/Source/Runtime/Niagara/Private/NiagaraSimulation.cpp b/Engine/Source/Runtime/Niagara/Private/NiagaraSimulation.cpp
--- a/Engine/Source/Runtime/Niagara/Private/NiagaraSimulation.cpp
+++ b/Engine/Source/Runtime/Niagara/Private/NiagaraSimulation.cpp
@@ -5,9 +5,9 @@
 #include "NiagaraEffectRenderer.h"
 #include "VectorVM.h"
 
-
 FNiagaraSimulation::FNiagaraSimulation(FNiagaraEmitterProperties *InProps) 
 : Age(0.0f)
+, Loops(0)
 , bIsEnabled(true)
 , SpawnRemainder(0.0f)
 , CachedBounds(ForceInit)
@@ -20,6 +20,7 @@ FNiagaraSimulation::FNiagaraSimulation(FNiagaraEmitterProperties *InProps)
 
 FNiagaraSimulation::FNiagaraSimulation(FNiagaraEmitterProperties *InProps, ERHIFeatureLevel::Type InFeatureLevel)
 	: Age(0.0f)
+	, Loops(0)
 	, bIsEnabled(true)
 	, SpawnRemainder(0.0f)
 	, CachedBounds(ForceInit)
@@ -115,7 +116,7 @@ void FNiagaraSimulation::Tick(float DeltaSeconds)
 {
 	SCOPE_CYCLE_COUNTER(STAT_NiagaraTick);
 
-	if (!bIsEnabled)
+	if (!bIsEnabled || TickState==NTS_Suspended || TickState==NTS_Dead)
 		return;
 
 	SimpleTimer TickTime;
@@ -143,12 +144,14 @@ void FNiagaraSimulation::Tick(float DeltaSeconds)
 	Constants.SetOrAdd(TEXT("Delta Time"), FVector4(DeltaSeconds, DeltaSeconds, DeltaSeconds, DeltaSeconds));
 
 	// Simulate particles forward by DeltaSeconds.
+	if (TickState==NTS_Running || TickState==NTS_Dieing)
 	{
 		SCOPE_CYCLE_COUNTER(STAT_NiagaraSimulate);
 		RunVMScript(Props->UpdateScript, EUnusedAttributeBehaviour::Copy);
 	}
 	
 	//Init new particles with the spawn script.
+	if (TickState==NTS_Running)
 	{
 		SCOPE_CYCLE_COUNTER(STAT_NiagaraSpawn);
 		Data.SetNumParticles(MaxNewParticles);
@@ -175,8 +178,16 @@ void FNiagaraSimulation::Tick(float DeltaSeconds)
 			}
 		}
 		Data.SetNumParticles(CurNumParticles);
+
+		// check if the emitter has officially died
+		if (GetTickState() == NTS_Dieing && CurNumParticles == 0)
+		{
+			SetTickState(NTS_Dead);
+		}
 	}
 
+
+
 	CPUTimeMS = TickTime.GetElapsedMilliseconds();
 
 	DECLARE_DWORD_COUNTER_STAT(TEXT("NumParticles"), STAT_NiagaraNumParticles, STATGROUP_Niagara);
diff --git a/Engine/Source/Runtime/Niagara/Public/NiagaraCommon.h b/Engine/Source/Runtime/Niagara/Public/NiagaraCommon.h
--- a/Engine/Source/Runtime/Niagara/Public/NiagaraCommon.h
+++ b/Engine/Source/Runtime/Niagara/Public/NiagaraCommon.h
@@ -15,6 +15,18 @@ enum class ENiagaraDataType : uint8
 	Curve,
 };
 
+
+
+
+enum ENiagaraTickState
+{
+	NTS_Running,
+	NTS_Suspended,
+	NTS_Dieing,
+	NTS_Dead
+};
+
+
 USTRUCT()
 struct FNiagaraVariableInfo
 {
