Making the effect editor work again.

diff --git a/Engine/Source/Editor/NiagaraEditor/Private/NiagaraSequencer.h b/Engine/Source/Editor/NiagaraEditor/Private/NiagaraSequencer.h
--- a/Engine/Source/Editor/NiagaraEditor/Private/NiagaraSequencer.h
+++ b/Engine/Source/Editor/NiagaraEditor/Private/NiagaraSequencer.h
@@ -25,6 +25,7 @@ public:
 	void SetEmitterProps(TWeakObjectPtr<const UNiagaraEmitterProperties> Props)	{ EmitterProps = Props; }
 
 	TWeakObjectPtr<const UNiagaraEmitterProperties> GetEmitterProps() const	{ return EmitterProps; }
+
 private:
 	FText EmitterName;
 	TWeakObjectPtr<const UNiagaraEmitterProperties> EmitterProps;
@@ -124,6 +125,7 @@ public:
 	virtual void ClearInstance( IMovieScenePlayer& Player ) override {}
 	virtual void SaveState(const TArray<UObject*>& RuntimeObjects) override {}
 	virtual void RestoreState(const TArray<UObject*>& RuntimeObjects) override {}
+
 private:
 	TSharedPtr<FNiagaraSimulation> Emitter;
 	class UEmitterMovieSceneTrack *Track;
@@ -175,6 +177,11 @@ public:
 		return Sections;
 	}
 
+	virtual TRange<float> GetSectionBoundaries() const override
+	{
+		return TRange<float>(0, FLT_MAX);
+	}
+
 	TSharedPtr<FNiagaraSimulation> GetEmitter() { return Emitter; }
 
 private:
