Fix Niagara track sections not displaying section borders

diff --git a/Engine/Source/Editor/NiagaraEditor/Private/NiagaraSequencer.h b/Engine/Source/Editor/NiagaraEditor/Private/NiagaraSequencer.h
--- a/Engine/Source/Editor/NiagaraEditor/Private/NiagaraSequencer.h
+++ b/Engine/Source/Editor/NiagaraEditor/Private/NiagaraSequencer.h
@@ -88,7 +88,7 @@ public:
 				);
 		}
 
-		return 0;
+		return LayerId;
 	}
 
 	FText GetDisplayName(void) const	{ return EmitterSection->GetEmitterName(); }
