** PRIMARY FEATURE **
- Events
Events work by writing to a preallocated buffer exposed to the VM as "SharedData". The underlying data storage for these is the same as particle attriubtes. Each "Shared Data" view in the VM has a counter which is incremented or decremented to aqcuire and consume a slot in this buffer for any particle that wished to write or read from it.
In the graph, this is done via ReadEvent and WriteEvent nodes. Each lets the user define the data they wish to read or write, just like for particle attributes.
Each also has a mandatory "Valid" pin which indicates whether the data being read or written is valid.
e.g.
A write would be marked deliberately valid or invalid by the user to control which particles write out event data. Valid <= 0.0f means dont' write. Valid > 0.0f means do write. Matches current working of less than etc. Should probably change to use ints when they're supported.
A read would be valid or not depending on whether there is data available in the event buffer.

Events TODO:
Once ints are supported, alot more power can be had using direct indexing of the buffers.
This leads me onto supporting generic, non-event buffers. Should be very easy to do.
Much of the VM support is hacked in and will want to be replaced after scalarization.
I may also move away form a "Shared Data" area and write a more general array feature. Couldn't really do this initially with no int support. Should have really done scalarization+int support first.
Needs BP interface.

** SECONDARY FEATURES **
- VM Debugger working - Still basic and rough. Allows a single particle to be traced through the VM for it's lifetime.
- Function Merging changes - Merged as needed during compiler traversal rather than all at start.
- Constatnt changes - Splat scalars when building table, not runtime. Faster and reduced code complexity.
- Fixed reallocation of niagara node pins so that existing connections are preserved.
- Added a reregister context class to reregister all niagara components that use a particualr script when that script is recompiled. TODO: Should also do this for property changes in the editor I think.
- Implemented GreaterThan Op. The reverse of LessThan
- Implemented VectorSelect Op.  > 0.0f Selects input A, otherwise B.
- Added a new Graph Editor Tab to the Niagara Graph Editor. "Flattened Graph" shows you the scripts graph but with all the function calls merged into the main graph. Useful for debugging. Still a bit rough around the edges. Might not make it to V1.0 of niagara but useful for the moment.

** MINOR FIXES **
Fixed crash in Dev UI in the effect editor.
Tweakd the details customization for the dev ui.
Fixed leakage of all effect instances that were created. This stuff still needs some work. Need to do a pass over the high level stuff in Niagara and more carefully define lifetimes and ownerships.

diff --git a/Engine/Source/Editor/NiagaraEditor/Private/NiagaraCompiler.cpp b/Engine/Source/Editor/NiagaraEditor/Private/NiagaraCompiler.cpp
--- a/Engine/Source/Editor/NiagaraEditor/Private/NiagaraCompiler.cpp
+++ b/Engine/Source/Editor/NiagaraEditor/Private/NiagaraCompiler.cpp
@@ -12,7 +12,7 @@
 
 #define LOCTEXT_NAMESPACE "NiagaraCompiler"
 
-DEFINE_LOG_CATEGORY_STATIC(LogNiagaraCompiler,All,All);
+DEFINE_LOG_CATEGORY_STATIC(LogNiagaraCompiler, All, All);
 
 void FNiagaraEditorModule::CompileScript(UNiagaraScript* ScriptToCompile)
 {
@@ -44,10 +44,11 @@ public:
 		ResultLocation = ENiagaraExpressionResultLocation::InputData;
 	}
 
-	virtual void Process()override
+	virtual bool Process()override
 	{
 		check(ResultLocation == ENiagaraExpressionResultLocation::InputData);
 		ResultIndex = Compiler->GetAttributeIndex(Result);
+		return true;
 	}
 };
 
@@ -78,12 +79,12 @@ public:
 // 		}
 	}
 
-	virtual void Process()override
+	virtual bool Process()override
 	{
 		check(ResultLocation == ENiagaraExpressionResultLocation::Constants);
 		//Get the index of the constant. Also gets a component index if this is for a packed scalar constant.
-		Compiler->GetConstantResultIndex(Result, bInternal, ResultIndex, ComponentIndex);
-		
+		Compiler->GetConstantResultIndex(Result, bInternal, ResultIndex);
+
 		if (Result.Type == ENiagaraDataType::Matrix)
 		{
 			check(SourceExpressions.Num() == 4);
@@ -104,6 +105,8 @@ public:
 		{
 			ResultLocation = ENiagaraExpressionResultLocation::BufferConstants;
 		}
+
+		return true;
 	}
 
 	bool bInternal;
@@ -188,6 +191,56 @@ int32 FNiagaraCompiler::GetAttributeIndex(const FNiagaraVariableInfo& Attr)const
 	return SourceGraph->GetAttributeIndex(Attr);
 }
 
+FNiagaraDataSetProperties* FNiagaraCompiler::GetSharedDataIndex(FNiagaraDataSetID DataSet, bool bForRead, int32& OutIndex)
+{
+	//Shared data is laid out:
+	//EventReceivers
+	//EventGenerators
+
+	int32 SetIndex = 0;
+	FNiagaraDataSetProperties* Ret = NULL;
+	auto SearchSharedArea = [&](TArray<FNiagaraDataSetProperties>& SharedSetGroup) -> FNiagaraDataSetProperties*
+	{
+		for (FNiagaraDataSetProperties& Props : SharedSetGroup)
+		{
+			if (DataSet == Props.ID)
+			{
+				OutIndex = SetIndex;
+				return &Props;
+			}
+			++SetIndex;
+		}
+		return NULL;
+	};
+
+	if (bForRead)
+	{
+		Ret = SearchSharedArea(EventReceivers);
+	}
+	else
+	{
+		SetIndex = EventReceivers.Num();
+		Ret = SearchSharedArea(EventGenerators);
+	}
+	check(Ret);
+	return Ret;
+}
+
+int32 FNiagaraCompiler::GetSharedDataVariableIndex(FNiagaraDataSetProperties* SharedDataSet, const FNiagaraVariableInfo& Variable)
+{
+	check(SharedDataSet);
+
+	for (int32 i = 0; i < SharedDataSet->Variables.Num(); ++i)
+	{
+		if (SharedDataSet->Variables[i] == Variable)
+		{
+			return i;
+		}
+	}
+
+	return INDEX_NONE;
+}
+
 TNiagaraExprPtr FNiagaraCompiler::CompileOutputPin(UEdGraphPin* Pin)
 {
 	check(Pin->Direction == EGPD_Output);
@@ -205,7 +258,7 @@ TNiagaraExprPtr FNiagaraCompiler::CompileOutputPin(UEdGraphPin* Pin)
 		UNiagaraNode* Node = Cast<UNiagaraNode>(Pin->GetOwningNode());
 		TArray<FNiagaraNodeResult> Outputs;
 		Node->Compile(this, Outputs);
-		for (int32 i=0; i<Outputs.Num(); ++i)
+		for (int32 i = 0; i < Outputs.Num(); ++i)
 		{
 			PinToExpression.Add(Outputs[i].OutputPin, Outputs[i].Expression);
 			//We also grab the expression for the pin we're currently interested in. Otherwise we'd have to search the map for it.
@@ -224,7 +277,19 @@ TNiagaraExprPtr FNiagaraCompiler::CompilePin(UEdGraphPin* Pin)
 	{
 		if (Pin->LinkedTo.Num() > 0)
 		{
-			// > 1 ???
+			//If we link to a function call, merge it's graph in.
+			TArray<UNiagaraScript*> Stack;
+			UNiagaraNodeFunctionCall* Function = Cast<UNiagaraNodeFunctionCall>(Pin->LinkedTo[0]->GetOwningNode());
+			while(Function)
+			{
+				if (!MergeFunctionIntoMainGraph(Function, Stack))
+				{
+					return nullptr;
+				}
+				
+				Function = Cast<UNiagaraNodeFunctionCall>(Pin->LinkedTo[0]->GetOwningNode());
+			}
+			//Any functon calls linked to this pin are now merged into the main graph, compile it.
 			Ret = CompileOutputPin(Pin->LinkedTo[0]);
 		}
 		else if (!Pin->bDefaultValueIsIgnored)
@@ -287,6 +352,8 @@ void FNiagaraCompiler::GetParticleAttributes(TArray<FNiagaraVariableInfo>& OutAt
 
 TNiagaraExprPtr FNiagaraCompiler::Expression_GetAttribute(const FNiagaraVariableInfo& Attribute)
 {
+	Usage.bReadsAttriubteData = true;
+
 	int32 Index = Expressions.Add(MakeShareable(new FNiagaraExpression_GetAttribute(this, Attribute)));
 	return Expressions[Index];
 }
@@ -359,8 +426,7 @@ TNiagaraExprPtr FNiagaraCompiler::GetInternalConstant(const FNiagaraVariableInfo
 	return Expression_GetInternalConstant(Constant);
 }
 
-// TODO: Refactor this (and some other stuff) into a preprocessing step for use by any compiler?
-bool FNiagaraCompiler::MergeInFunctionNodes()
+bool FNiagaraCompiler::MergeFunctionIntoMainGraph(UNiagaraNodeFunctionCall* FunctionNode, TArray<UNiagaraScript*>& FunctionStack)
 {
 	struct FReconnectionInfo
 	{
@@ -375,295 +441,228 @@ bool FNiagaraCompiler::MergeInFunctionNodes()
 			: From(NULL)
 		{}
 	};
-	TMap<FName, FReconnectionInfo> InputConnections;
-	TMap<FName, FReconnectionInfo> OutputConnections;
 
-	TArray<class UEdGraphPin*> FuncCallInputPins;
-	TArray<class UEdGraphPin*> FuncCallOutputPins;
-
-	//Copies the function graph into the main graph.
-	//Removes the Function call in the main graph and the input and output nodes in the function graph, reconnecting their pins appropriately.
-	auto MergeFunctionIntoMainGraph = [&](UNiagaraNodeFunctionCall* InFunc, UNiagaraGraph* FuncGraph)
+	if (UNiagaraScript* FuncScript = FunctionNode->FunctionScript)
 	{
-		InputConnections.Empty();
-		OutputConnections.Empty();
-		FuncCallInputPins.Empty();
-		FuncCallOutputPins.Empty();
+		//Check for reentrance.
+		if (FunctionStack.Find(FuncScript) != INDEX_NONE)
+		{
+			FString Callstack;
+			Callstack.Append(FuncScript->GetPathName()); Callstack.Append(TEXT("\n"));
+			for (UNiagaraScript* StackScript : FunctionStack)
+			{
+				Callstack.Append(StackScript->GetPathName()); Callstack.Append(TEXT("\n"));
+			}
+			MessageLog.Error(TEXT("Reentrant function call!\n%s"), *Callstack);
+			return false;
+		}
+		FunctionStack.Add(FuncScript);
+
+		// Clone the source graph so we can modify it as needed; merging in the child graphs
+		UNiagaraScriptSource* FuncSource = CastChecked<UNiagaraScriptSource>(FuncScript->Source);
+		UNiagaraGraph* FuncGraph = CastChecked<UNiagaraGraph>(FEdGraphUtilities::CloneGraph(FuncSource->NodeGraph, NULL, &MessageLog));
+		FEdGraphUtilities::MergeChildrenGraphsIn(FuncGraph, FuncGraph, /*bRequireSchemaMatch=*/ true);
+
+		//Copies the function graph into the main graph.
+		//Removes the Function call in the main graph and the input and output nodes in the function graph, reconnecting their pins appropriately.
+		//auto MergeFunctionIntoMainGraph = [&](UNiagaraNodeFunctionCall* InFunc, UNiagaraGraph* FuncGraph)
 
-		check(InFunc && FuncGraph);
-		if (InFunc->FunctionScript)
+		TMap<FName, FReconnectionInfo> InputConnections;
+		TMap<FName, FReconnectionInfo> OutputConnections;
+		TArray<class UEdGraphPin*> FuncCallInputPins;
+		TArray<class UEdGraphPin*> FuncCallOutputPins;
+
+		//Get all the pins that are connected to the inputs of the function call node in the main graph.
+		FunctionNode->GetInputPins(FuncCallInputPins);
+		for (UEdGraphPin* FuncCallInputPin : FuncCallInputPins)
 		{
-			//Get all the pins that are connected to the inputs of the function call node in the main graph.
-			InFunc->GetInputPins(FuncCallInputPins);
-			for (UEdGraphPin* FuncCallInputPin : FuncCallInputPins)
+			FName InputName(*FuncCallInputPin->PinName);
+			FReconnectionInfo& InputConnection = InputConnections.FindOrAdd(InputName);
+			if (FuncCallInputPin->LinkedTo.Num() > 0)
 			{
-				FName InputName(*FuncCallInputPin->PinName);
-				FReconnectionInfo& InputConnection = InputConnections.FindOrAdd(InputName);
-				if (FuncCallInputPin->LinkedTo.Num() > 0)
-				{
-					check(FuncCallInputPin->LinkedTo.Num() == 1);
-					UEdGraphPin* LinkFrom = FuncCallInputPin->LinkedTo[0];
-					check(LinkFrom->Direction == EGPD_Output);
-					InputConnection.From = LinkFrom;
-				}
-				else
-				{
-					//This input has no link so we need the default value from the pin.
-					InputConnection.FallbackDefault = FuncCallInputPin->GetDefaultAsString();
-				}
+				check(FuncCallInputPin->LinkedTo.Num() == 1);
+				UEdGraphPin* LinkFrom = FuncCallInputPin->LinkedTo[0];
+				check(LinkFrom->Direction == EGPD_Output);
+				InputConnection.From = LinkFrom;
+			}
+			else
+			{
+				//This input has no link so we need the default value from the pin.
+				InputConnection.FallbackDefault = FuncCallInputPin->GetDefaultAsString();
 			}
-			//Get all the pins that are connected to the outputs of the function call node in the main graph.
-			InFunc->GetOutputPins(FuncCallOutputPins);
-			for (UEdGraphPin* FuncCallOutputPin : FuncCallOutputPins)
+		}
+		//Get all the pins that are connected to the outputs of the function call node in the main graph.
+		FunctionNode->GetOutputPins(FuncCallOutputPins);
+		for (UEdGraphPin* FuncCallOutputPin : FuncCallOutputPins)
+		{
+			FName OutputName(*FuncCallOutputPin->PinName);
+			for (UEdGraphPin* LinkTo : FuncCallOutputPin->LinkedTo)
 			{
-				FName OutputName(*FuncCallOutputPin->PinName);
-				for (UEdGraphPin* LinkTo : FuncCallOutputPin->LinkedTo)
+				if (LinkTo)
 				{
-					if (LinkTo)
-					{
-						check(LinkTo->Direction == EGPD_Input);
-						FReconnectionInfo& OutputConnection = OutputConnections.FindOrAdd(OutputName);
-						OutputConnection.To.Add(LinkTo);
-					}
+					check(LinkTo->Direction == EGPD_Input);
+					FReconnectionInfo& OutputConnection = OutputConnections.FindOrAdd(OutputName);
+					OutputConnection.To.Add(LinkTo);
 				}
 			}
+		}
 
-			//Remove the function call node from the graph now that we have everything we need from it.
-			SourceGraph->RemoveNode(InFunc);
+		//Remove the function call node from the graph now that we have everything we need from it.
+		SourceGraph->RemoveNode(FunctionNode);
+		UNiagaraNodeOutput* SourceGraphOuputNode = SourceGraph->FindOutputNode();
+		check(SourceGraphOuputNode);
 
-			//Keep a list of the Input and Output nodes we see in the function graph so that we can remove (most of) them later.
-			TArray<UEdGraphNode*, TInlineAllocator<64>> ToRemove;
+		//Keep a list of the Input and Output nodes we see in the function graph so that we can remove (most of) them later.
+		TArray<UEdGraphNode*, TInlineAllocator<64>> ToRemove;
 
-			//Search the nodes in the function graph, finding any connections to input or output nodes.
-			for (UEdGraphNode* FuncGraphNode : FuncGraph->Nodes)
+		//For each input node in the function graph, get all the LinkTo connections for reconnecting to the main graph later.
+		TArray <UNiagaraNodeInput*> FuncGraphInputNodes;
+		FuncGraph->FindInputNodes(FuncGraphInputNodes);
+		for (UNiagaraNodeInput* FuncGraphInputNode : FuncGraphInputNodes)
+		{
+			check(FuncGraphInputNode && FuncGraphInputNode->Pins.Num() == 1);
+			//Get an array of "To" pins from one or more input nodes referencing each named input.
+			FReconnectionInfo& InputConnection = InputConnections.FindOrAdd(FuncGraphInputNode->Input.Name);
+			if (InputConnection.From)
 			{
-				if (UNiagaraNodeInput* InputNode = Cast<UNiagaraNodeInput>(FuncGraphNode))
-				{
-					check(InputNode->Pins.Num() == 1);
-					//Get an array of "To" pins from one or more input nodes referencing each named input.
-					FReconnectionInfo& InputConnection = InputConnections.FindOrAdd(InputNode->Input.Name);
-					if (InputConnection.From)
-					{
-						//We have a connection from the function call so remove the input node and connect to that.
-						ToRemove.Add(InputNode);
-					}
-					else
-					{
-						//This input has no connection from the function call so what do we do here? 
-						//For now we just leave the input node and connect back to it. 
-						//This will mean unconnected pins from the function call will look for constants or attributes. 
-						//In some cases we may want to just take the default value from the function call pin instead?
-						//Maybe have some properties on the function call defining that.
-						InputConnection.From = InputNode->Pins[0];
-					}
-
-					TArray<UEdGraphPin*>& LinkToPins = InputNode->Pins[0]->LinkedTo;
-					for (UEdGraphPin* ToPin : LinkToPins)
-					{
-						check(ToPin->Direction == EGPD_Input);
-						InputConnection.To.Add(ToPin);
-					}
-				}
-				else if (UNiagaraNodeOutput* OutputNode = Cast<UNiagaraNodeOutput>(FuncGraphNode))
-				{
-					//Unlike the input nodes, we don't have the option of keeping these if there is no "From" pin. The default values from the node pins should be used.
-					ToRemove.Add(OutputNode);
-
-					//For each output, get the "From" pin to be reconnected later.
-					for (int32 OutputIdx = 0; OutputIdx < OutputNode->Outputs.Num(); ++OutputIdx)
-					{
-						FName OutputName = OutputNode->Outputs[OutputIdx].Name;
-
-						UEdGraphPin* OutputNodePin = OutputNode->Pins[OutputIdx];
-						check(OutputNodePin->LinkedTo.Num() <= 1);
-						FReconnectionInfo& OutputConnection = OutputConnections.FindOrAdd(OutputName);
-						UEdGraphPin* LinkFromPin = OutputNodePin->LinkedTo.Num() == 1 ? OutputNodePin->LinkedTo[0] : NULL;
-						if (LinkFromPin)
-						{
-							check(LinkFromPin->Direction == EGPD_Output);
-							OutputConnection.From = LinkFromPin;
-						}
-						else
-						{
-							//This output is not connected so links to it in the main graph must use it's default value.
-							OutputConnection.FallbackDefault = OutputNodePin->GetDefaultAsString();
-						}
-					}
-				}
+				//We have a connection from the function call so remove the input node and connect to that.
+				ToRemove.Add(FuncGraphInputNode);
 			}
-
-			//Remove all the In and Out nodes from the function graph.
-			for (UEdGraphNode* Remove : ToRemove)
+			else
 			{
-				FuncGraph->RemoveNode(Remove);
+				//This input has no connection from the function call so what do we do here? 
+				//For now we just leave the input node and connect back to it. 
+				//This will mean unconnected pins from the function call will look for constants or attributes. 
+				//In some cases we may want to just take the default value from the function call pin instead?
+				//Maybe have some properties on the function call defining that.
+				InputConnection.From = FuncGraphInputNode->Pins[0];
 			}
 
-			//Copy the nodes from the function graph over into the main graph.
-			FuncGraph->MoveNodesToAnotherGraph(SourceGraph, false);
-
-			//Finally, do all the reconnection.
-			auto MakeConnection = [&](FReconnectionInfo& Info)
+			TArray<UEdGraphPin*>& LinkToPins = FuncGraphInputNode->Pins[0]->LinkedTo;
+			for (UEdGraphPin* ToPin : LinkToPins)
 			{
-				for (UEdGraphPin* LinkTo : Info.To)
+				check(ToPin->Direction == EGPD_Input);
+				UEdGraphNode* OwningNode = ToPin->GetOwningNode();
+				if (!OwningNode->IsA(UNiagaraNodeOutput::StaticClass()))
 				{
-					if (Info.From)
-					{
-						Info.From->MakeLinkTo(LinkTo);
-					}
-					else
-					{
-						LinkTo->DefaultValue = Info.FallbackDefault;
-					}
+					//As long as the owning node will survive into the source graph, connect directly to it.
+					//Otherwise, we handle the connection in the OutputConnections.
+					InputConnection.To.Add(ToPin);
 				}
-			};
-			for (TPair<FName, FReconnectionInfo>& ReconnectInfo : InputConnections){ MakeConnection(ReconnectInfo.Value); }
-			for (TPair<FName, FReconnectionInfo>& ReconnectInfo : OutputConnections){ MakeConnection(ReconnectInfo.Value); }
+			}
 		}
-	};
 
-	//Helper struct for traversing nested function calls.
-	struct FFunctionContext
-	{
-		//True if this context's function has been merged into the main graph.
-		bool bProcessed;
-		//The index of this context into the ContextPool. 
-		int32 PoolIdx;
-		//Pointer back to the parent context for traversal.
-		FFunctionContext* Parent;
-		//The function call node for this function in the source/parent graph.
-		UNiagaraNodeFunctionCall* Function;
-		//The graph for this function that we are going to merge into the main graph.
-		UNiagaraGraph* FunctionGraph;
-		//The script from which the graph is copied. Used for re entrance check.
-		UNiagaraScript* Script;
-
-		//Contexts for function calls in this function graph.
-		TArray<FFunctionContext*, TInlineAllocator<64>> SubFunctionCalls;
-
-		FFunctionContext()
-			: bProcessed(false)
-			, PoolIdx(INDEX_NONE)
-			, Parent(NULL)
-			, Function(NULL)
-			, FunctionGraph(NULL)
-			, Script(NULL)
+		//For each output of the output node in the function graph, get the "From" pin to be reconnected later.
 		{
-		}
+			UNiagaraNodeOutput* FuncGraphOutputNode = FuncGraph->FindOutputNode();
+			check(FuncGraphOutputNode);
 
-		/**
-		We don't allow re-entrant functions as this would cause an infinite loop of merging in graphs.
-		Maybe in the future if we allow branching in the VM we can allow this.
-		*/
-		bool CheckForReentrance()const
-		{
-			UNiagaraNodeFunctionCall* Func = Function;
-			FFunctionContext* Curr = Parent;
-			while (Curr)
+			//Unlike the input nodes, we don't have the option of keeping these if there is no "From" pin. The default values from the node pins should be used.
+			ToRemove.Add(FuncGraphOutputNode);
+
+			for (int32 OutputIdx = 0; OutputIdx < FuncGraphOutputNode->Outputs.Num(); ++OutputIdx)
 			{
-				if (Curr->Script == Script)
-					return true;
+				FName OutputName = FuncGraphOutputNode->Outputs[OutputIdx].Name;
 
-				Curr = Curr->Parent;
-			}
-			return false;
-		}
+				UEdGraphPin* OutputNodePin = FuncGraphOutputNode->Pins[OutputIdx];
+				check(OutputNodePin->LinkedTo.Num() <= 1);
+				FReconnectionInfo& OutputConnection = OutputConnections.FindOrAdd(OutputName);
+				UEdGraphPin* LinkFromPin = OutputNodePin->LinkedTo.Num() == 1 ? OutputNodePin->LinkedTo[0] : NULL;
 
-		FString GetCallstack()const
-		{
-			FString Ret;
-			const FFunctionContext* Curr = this;
-			while (Curr)
-			{
-				if (Curr->Script)
+				if (LinkFromPin)
+				{
+					UEdGraphNode* OwningNode = LinkFromPin->GetOwningNode();
+					if (UNiagaraNodeInput* InputNode = Cast<UNiagaraNodeInput>(OwningNode))
+					{
+						//We are connected directly to an input node so we must find the corresponding input connection and connect to that.
+						if (FReconnectionInfo* InputConnection = InputConnections.Find(InputNode->Input.Name))
+						{
+							LinkFromPin = InputConnection->From;
+						}
+					}
+				}
+
+				if (LinkFromPin)
 				{
-					Ret.Append(*(Curr->Script->GetPathName()));
+					check(LinkFromPin->Direction == EGPD_Output);
+					OutputConnection.From = LinkFromPin;
 				}
 				else
 				{
-					Ret.Append(TEXT("Unknown"));
+					//This output is not connected so links to it in the main graph must use it's default value.
+					OutputConnection.FallbackDefault = OutputNodePin->GetDefaultAsString();
 				}
-
-				Ret.Append(TEXT("\n"));
-
-				Curr = Curr->Parent;
 			}
-			return Ret;
 		}
-	};
 
-	//A pool of contexts on the stack to avoid loads of needless, small heap allocations.
-	TArray<FFunctionContext, TInlineAllocator<512>> ContextPool;
-	ContextPool.Reserve(512);
+		//Remove all the In and Out nodes from the function graph.
+		for (UEdGraphNode* Remove : ToRemove)
+		{
+			FuncGraph->RemoveNode(Remove);
+		}
+
+		TArray<UEdGraphNode*> FuncNodes = FuncGraph->Nodes;
+		TArray<UEdGraphNode*> ExistingNodes = SourceGraph->Nodes;
+		//Copy the nodes from the function graph over into the main graph.
+		FuncGraph->MoveNodesToAnotherGraph(SourceGraph, false);
 
-	FFunctionContext RootContext;
-	FFunctionContext* CurrentContext = &RootContext;
-	CurrentContext->FunctionGraph = SourceGraph;
-	CurrentContext->Script = Script;
+		FIntPoint FuncNodePos(FunctionNode->NodePosX, FunctionNode->NodePosY);
+		FIntRect Bounds = FEdGraphUtilities::CalculateApproximateNodeBoundaries(FuncNodes);
+		FIntPoint Center;
+		FIntPoint Extents;
+		Bounds.GetCenterAndExtents(Center, Extents);
 
-	//Depth first traversal of all function calls.
-	while (CurrentContext)
-	{
-		//Find any sub functions and process this function call.
-		if (!CurrentContext->bProcessed)
+		FIntPoint Offset = FuncNodePos - Center;
+		for (UEdGraphNode* FuncNode : FuncNodes)
 		{
-			CurrentContext->bProcessed = true;
+			FuncNode->NodePosX += Offset.X;
+			FuncNode->NodePosY += Offset.Y;
+		}
 
-			//Find any sub functions and check for re-entrance.
-			if (CurrentContext->FunctionGraph)
+		for (UEdGraphNode* Node : ExistingNodes)
+		{
+			if (Node->NodePosX < FuncNodePos.X)
 			{
-				for (UEdGraphNode* Node : CurrentContext->FunctionGraph->Nodes)
-				{
-					UNiagaraNodeFunctionCall* FuncNode = Cast<UNiagaraNodeFunctionCall>(Node);
-					if (FuncNode)
-					{
-						int32 NewIdx = ContextPool.AddZeroed();
-						FFunctionContext* SubFuncContext = &ContextPool[NewIdx];
-						CurrentContext->SubFunctionCalls.Push(SubFuncContext);
-						SubFuncContext->Parent = CurrentContext;
-						SubFuncContext->Function = FuncNode;
-						SubFuncContext->PoolIdx = NewIdx;
-						SubFuncContext->Script = FuncNode->FunctionScript;
-
-						if (SubFuncContext->CheckForReentrance())
-						{
-							FString Callstack = SubFuncContext->GetCallstack();
-							MessageLog.Error(TEXT("Reentrant function call!\n%s"), *Callstack);
-							return false;
-						}
-
-						//Copy the function graph as we'll be modifying it as we merge in with the main graph.
-						UNiagaraScriptSource* FuncSource = CastChecked<UNiagaraScriptSource>(FuncNode->FunctionScript->Source);
-						check(FuncSource);
-						SubFuncContext->FunctionGraph = CastChecked<UNiagaraGraph>(FEdGraphUtilities::CloneGraph(FuncSource->NodeGraph, NULL, &MessageLog));
-					}
-				}
+				Node->NodePosX -= Extents.X;
+			}
+			else
+			{
+				Node->NodePosX += Extents.X;
 			}
 
-			//Merge this function into the main graph now.
-			if (CurrentContext->Function && CurrentContext->FunctionGraph)
+			if (Node->NodePosY < FuncNodePos.Y)
+			{
+				Node->NodePosY -= Extents.Y;
+			}
+			else
 			{
-				MergeFunctionIntoMainGraph(CurrentContext->Function, CurrentContext->FunctionGraph);
+				Node->NodePosY += Extents.Y;
 			}
 		}
 
-		if (CurrentContext->SubFunctionCalls.Num() > 0)
+		//Finally, do all the reconnection.
+		auto MakeConnection = [&](FReconnectionInfo& Info)
 		{
-			//Move to the next sub function.
-			CurrentContext = CurrentContext->SubFunctionCalls.Pop();
-		}
-		else
-		{
-			//Done processing this function so remove it and move back to the parent.
-			if (CurrentContext->PoolIdx != INDEX_NONE)
+			for (UEdGraphPin* LinkTo : Info.To)
 			{
-				CurrentContext->FunctionGraph->MarkPendingKill();
-
-				ContextPool.RemoveAtSwap(CurrentContext->PoolIdx);
+				if (Info.From)
+				{
+					Info.From->MakeLinkTo(LinkTo);
+				}
+				else
+				{
+					LinkTo->DefaultValue = Info.FallbackDefault;
+				}
 			}
-			CurrentContext = CurrentContext->Parent;
-		}
+		};
+		for (TPair<FName, FReconnectionInfo>& ReconnectInfo : InputConnections){ MakeConnection(ReconnectInfo.Value); }
+		for (TPair<FName, FReconnectionInfo>& ReconnectInfo : OutputConnections){ MakeConnection(ReconnectInfo.Value); }
+	}
+	else
+	{ 
+		Error(LOCTEXT("MergeNullFunction", "Function call node must reference a valid script."), FunctionNode, NULL);
+		return false;
 	}
-
 	return true;
 }
 
diff --git a/Engine/Source/Editor/NiagaraEditor/Private/NiagaraCompiler.h b/Engine/Source/Editor/NiagaraEditor/Private/NiagaraCompiler.h
--- a/Engine/Source/Editor/NiagaraEditor/Private/NiagaraCompiler.h
+++ b/Engine/Source/Editor/NiagaraEditor/Private/NiagaraCompiler.h
@@ -5,6 +5,7 @@
 #include "INiagaraCompiler.h"
 #include "CompilerResultsLog.h"
 #include "NiagaraScriptConstantData.h"
+#include "NiagaraScript.h"
 
 /** Base class for Niagara compilers. Children of this will include a compiler for the VectorVM and for Compute shaders. Possibly others. */
 class NIAGARAEDITOR_API FNiagaraCompiler : public INiagaraCompiler
@@ -22,15 +23,28 @@ protected:
 	/** The set of expressions generated from the script source. */
 	TArray<TNiagaraExprPtr> Expressions;
 
-	//All internal and external constants used in the graph.
+	/** All internal and external constants used in the graph. */
 	FNiagaraScriptConstantData ConstantData;
 
 	/** Map of Pins to expressions. Allows us to reuse expressions for pins that have already been compiled. */
 	TMap<UEdGraphPin*, TNiagaraExprPtr> PinToExpression;
 
-	//Message log. Automatically handles marking the NodeGraph with errors.
+	/** Message log. Automatically handles marking the NodeGraph with errors. */
 	FCompilerResultsLog MessageLog;
 
+	/** List of event names this script receives and the attributes they read. */
+	TArray<FNiagaraDataSetProperties> EventReceivers;
+
+	/** List of event names this script generates and the attributes they write. */
+	TArray<FNiagaraDataSetProperties> EventGenerators;
+
+	FNiagaraScriptUsageInfo Usage;
+
+// 	/** List of other shared data read */
+// 	TArray<FNiagaraDataSetCompilationInfo> SharedDataRead;
+// 	/** List of other shared data written */
+// 	TArray<FNiagaraDataSetCompilationInfo> SharedDataWritten;
+
 	/** Compiles an output Pin on a graph node. Caches the result for any future inputs connected to it. */
 	TNiagaraExprPtr CompileOutputPin(UEdGraphPin* Pin);
 
@@ -53,6 +67,8 @@ protected:
 	/** Searches for Function Call nodes and merges their graphs into the main graph. */
 	bool MergeInFunctionNodes();
 
+	bool MergeFunctionIntoMainGraph(UNiagaraNodeFunctionCall* FunctionNode, TArray<UNiagaraScript*>& FunctionStack);
+
 public:
 
 	//Begin INiagaraCompiler Interface
@@ -76,9 +92,14 @@ public:
 	//End INiagaraCompiler Interface
 	
 	/** Gets the index into a constants table of the constant specified by Name and bInternal. */
-	virtual ENiagaraDataType GetConstantResultIndex(const FNiagaraVariableInfo& Constant, bool bInternal, int32& OutResultIndex, int32& OutComponentIndex) = 0;
-	/**	Gets the index of an attribute. */
+	virtual ENiagaraDataType GetConstantResultIndex(const FNiagaraVariableInfo& Constant, bool bInternal, int32& OutResultIndex) = 0;
+	/**	Gets the index of an attribute in particle data. */
 	int32 GetAttributeIndex(const FNiagaraVariableInfo& Attr)const;
+	/**	Gets the index of a shared data view. */
+	FNiagaraDataSetProperties* GetSharedDataIndex(FNiagaraDataSetID SharedDataSet, bool bForRead, int32& OutIndex);
+	/**	Gets the index of a variable in shared data. */
+	int32 GetSharedDataVariableIndex(FNiagaraDataSetProperties* SharedDataSet, const FNiagaraVariableInfo& Attr);
+
 	virtual void GetParticleAttributes(TArray<FNiagaraVariableInfo>& OutAttributes)const;
 
 	/**	Gets the index of a free temporary location. */
diff --git a/Engine/Source/Editor/NiagaraEditor/Private/NiagaraCompiler_VectorVM.cpp b/Engine/Source/Editor/NiagaraEditor/Private/NiagaraCompiler_VectorVM.cpp
--- a/Engine/Source/Editor/NiagaraEditor/Private/NiagaraCompiler_VectorVM.cpp
+++ b/Engine/Source/Editor/NiagaraEditor/Private/NiagaraCompiler_VectorVM.cpp
@@ -17,9 +17,9 @@ DEFINE_LOG_CATEGORY_STATIC(LogNiagaraCompiler_VectorVM, All, All);
 class FNiagaraExpression_VMOperation : public FNiagaraExpression
 {
 public:
-	VectorVM::EOp OpCode;
+	EVectorVMOp OpCode;
 
-	FNiagaraExpression_VMOperation(FNiagaraCompiler* Compiler, const FNiagaraVariableInfo& InResult, TArray<TNiagaraExprPtr>& InputExpressions, VectorVM::EOp Op)
+	FNiagaraExpression_VMOperation(FNiagaraCompiler* Compiler, const FNiagaraVariableInfo& InResult, TArray<TNiagaraExprPtr>& InputExpressions, EVectorVMOp Op)
 		: FNiagaraExpression(Compiler, InResult)
 		, OpCode(Op)
 	{
@@ -34,56 +34,43 @@ public:
 		}
 	}
 
-	virtual void Process()override
+	virtual bool Process()override
 	{
 		FNiagaraCompiler_VectorVM* VMCompiler = (FNiagaraCompiler_VectorVM*)Compiler;
 
 		//Should be invalid at this point.
 		check(ResultIndex == INDEX_NONE);
-		check(ComponentIndex == INDEX_NONE);
 
-		//Special case where we're splatting but we don't know the component until now. E.g. Getting constants.
-		//Choose the component now based on the source expressions component index
-		if (OpCode == VectorVM::EOp::splatx)
+		ResultIndex = VMCompiler->AquireTemporary();
+		if (ResultIndex == INDEX_NONE)
 		{
-			check(SourceExpressions.Num() == 1);
-			TNiagaraExprPtr Src0 = GetSourceExpression(0);
-			if (Src0->ComponentIndex != INDEX_NONE)
-			{
-				switch (Src0->ComponentIndex)
-				{
-				case 0: OpCode = VectorVM::EOp::splatx; break;
-				case 1: OpCode = VectorVM::EOp::splaty; break;
-				case 2: OpCode = VectorVM::EOp::splatz; break;
-				case 3: OpCode = VectorVM::EOp::splatw; break;
-				}
-			}
+			return false;
 		}
 
-		ResultIndex = VMCompiler->AquireTemporary();
 		ResultLocation = ENiagaraExpressionResultLocation::Temporaries;
 		check(ResultIndex < VectorVM::NumTempRegisters);
 
 		VMCompiler->WriteCode((uint8)OpCode);
-		VMCompiler->WriteCode(VMCompiler->GetResultVMIndex( this ));
 
 		//Add the bitfield defining whether each source operand comes from a constant or a register.
 		int32 NumOperands = SourceExpressions.Num();
 		check(NumOperands <= 4);
 
-
-		uint8 OpTypeMask = VectorVM::CreateSrcOperandMask(NumOperands > 0 ? VMCompiler->GetVMOperandType(GetSourceExpression(0).Get()) : VectorVM::RegisterOperandType,
-			NumOperands > 1 ? VMCompiler->GetVMOperandType(GetSourceExpression(1).Get()) : VectorVM::RegisterOperandType,
-			NumOperands > 2 ? VMCompiler->GetVMOperandType(GetSourceExpression(2).Get()) : VectorVM::RegisterOperandType,
-			NumOperands > 3 ? VMCompiler->GetVMOperandType(GetSourceExpression(3).Get()) : VectorVM::RegisterOperandType);
+		uint8 OpTypeMask = VectorVM::CreateSrcOperandMask(NumOperands > 0 ? VMCompiler->GetVMOperandType(GetSourceExpression(0).Get()) : EVectorVMOperandLocation::TemporaryRegister,
+			NumOperands > 1 ? VMCompiler->GetVMOperandType(GetSourceExpression(1).Get()) : EVectorVMOperandLocation::TemporaryRegister,
+			NumOperands > 2 ? VMCompiler->GetVMOperandType(GetSourceExpression(2).Get()) : EVectorVMOperandLocation::TemporaryRegister,
+			NumOperands > 3 ? VMCompiler->GetVMOperandType(GetSourceExpression(3).Get()) : EVectorVMOperandLocation::TemporaryRegister);
 		VMCompiler->WriteCode(OpTypeMask);
 
+		//Write dest
+		VMCompiler->WriteCode(VMCompiler->GetResultVMIndex(this));
 
 		//Add the locations for each of the source operands.
 		for (int32 SrcIdx = 0; SrcIdx < SourceExpressions.Num(); ++SrcIdx)
 		{
 			VMCompiler->WriteCode(VMCompiler->GetResultVMIndex(GetSourceExpression(SrcIdx).Get()));
 		}
+		return true;
 	}
 };
 
@@ -100,7 +87,7 @@ public:
 		check(SourceExpressions.Num() == 1);
 	}
 
-	virtual void Process()override
+	virtual bool Process()override
 	{
 		FNiagaraCompiler_VectorVM* VMCompiler = (FNiagaraCompiler_VectorVM*)Compiler;
 		check(ResultLocation == ENiagaraExpressionResultLocation::OutputData);
@@ -112,13 +99,251 @@ public:
 		uint8 DestIndex = VMCompiler->GetResultVMIndex(this);
 		uint8 SrcIndex = VMCompiler->GetResultVMIndex(SrcExpr);
 
-		VMCompiler->WriteCode((uint8)VectorVM::EOp::output);
+		VMCompiler->WriteCode((uint8)EVectorVMOp::output);
+		VMCompiler->WriteCode(VectorVM::CreateSrcOperandMask(VMCompiler->GetVMOperandType(SrcExpr)));
 		VMCompiler->WriteCode(DestIndex);
-		VMCompiler->WriteCode(VectorVM::CreateSrcOperandMask( VMCompiler->GetVMOperandType(SrcExpr) ) );
 		VMCompiler->WriteCode(SrcIndex);
+		return true;
 	}
 };
 
+/** Expression to read from shared data. */
+class FNiagaraExpression_VMDataSetRead : public FNiagaraExpression
+{
+public:
+
+	FNiagaraExpression_VMDataSetRead(class FNiagaraCompiler* InCompiler, FNiagaraDataSetID InEvent, const FNiagaraVariableInfo& InVariable, TNiagaraExprPtr& IndexExpr)
+		: FNiagaraExpression(InCompiler, InVariable)
+		, Event(InEvent)
+	{
+		SourceExpressions.Add(IndexExpr);
+	}
+
+	~FNiagaraExpression_VMDataSetRead()
+	{
+		Compiler->FreeTemporary(ResultIndex);
+	}
+
+	virtual bool Process()override
+	{
+		FNiagaraCompiler_VectorVM* VMCompiler = (FNiagaraCompiler_VectorVM*)Compiler;
+		
+		ResultIndex = VMCompiler->AquireTemporary();
+		if (ResultIndex == INDEX_NONE)
+		{
+			return false;
+		}
+		ResultLocation = ENiagaraExpressionResultLocation::Temporaries;
+		check(ResultIndex < VectorVM::NumTempRegisters);
+		
+		uint8 DestIndex = VMCompiler->GetResultVMIndex(this);
+		int32 SetIndex;
+		FNiagaraDataSetProperties* SetProps = VMCompiler->GetSharedDataIndex(Event, true, SetIndex);
+		check(SetIndex <= 255);
+		uint8 VarIndex = VMCompiler->GetSharedDataVariableIndex(SetProps, Result);
+
+		check(SourceExpressions.Num() == 1);
+		FNiagaraExpression* IndexExpr = GetSourceExpression(0).Get();
+		uint8 IndexIndex = VMCompiler->GetResultVMIndex(IndexExpr);
+
+		VMCompiler->WriteCode((uint8)EVectorVMOp::shareddataread);
+		VMCompiler->WriteCode(DestIndex);//Result location
+
+		VMCompiler->WriteCode((int8)SetIndex);//Index of data set
+		VMCompiler->WriteCode(VarIndex);//Index of variable within data set view
+		VMCompiler->WriteCode(IndexIndex);//Index of data in variable buffer
+		return true;
+	}
+
+	FNiagaraDataSetID Event;
+};
+
+/** Expression to write to shared data. */
+class FNiagaraExpression_VMDataSetWrite : public FNiagaraExpression
+{
+public:
+
+	FNiagaraExpression_VMDataSetWrite(class FNiagaraCompiler* InCompiler, FNiagaraDataSetID InEvent, const FNiagaraVariableInfo& InVariable, TNiagaraExprPtr& IndexExpr, TNiagaraExprPtr& DataExpr)
+		: FNiagaraExpression(InCompiler, InVariable)
+		, Event(InEvent)
+	{
+		ResultLocation = ENiagaraExpressionResultLocation::SharedData;
+		SourceExpressions.Add(IndexExpr);
+		SourceExpressions.Add(DataExpr);
+	}
+
+	virtual bool Process()override
+	{
+		FNiagaraCompiler_VectorVM* VMCompiler = (FNiagaraCompiler_VectorVM*)Compiler;
+		check(ResultLocation == ENiagaraExpressionResultLocation::SharedData);
+		FNiagaraDataSetProperties* DataSetProps = VMCompiler->GetSharedDataIndex(Event, false, ResultIndex);
+		int32 VarIndex = VMCompiler->GetSharedDataVariableIndex(DataSetProps, Result);
+
+		FNiagaraExpression* IndexExpr = GetSourceExpression(0).Get();
+		FNiagaraExpression* DataExpr = GetSourceExpression(1).Get();
+		check(SourceExpressions.Num() == 2);
+
+		uint8 IndexIndex = VMCompiler->GetResultVMIndex(IndexExpr);
+		uint8 DataIndex = VMCompiler->GetResultVMIndex(DataExpr);		
+
+		//TODO - Make this better, either ensure this is in a register by now or allow the VM to handle from constants?
+		//Though, why you'd have the conditional come from a constant, IDK.
+		if (!(DataExpr->ResultLocation == ENiagaraExpressionResultLocation::Temporaries || DataExpr->ResultLocation == ENiagaraExpressionResultLocation::InputData))
+		{
+			VMCompiler->Error(LOCTEXT("CompileFaliure_CannotWriteConstantToSharedData", "Writing of constants to shared data is not currently supported"), NULL, NULL);
+			return false;
+		}
+		check(IndexExpr->ResultLocation == ENiagaraExpressionResultLocation::Temporaries);
+
+		VMCompiler->WriteCode((uint8)EVectorVMOp::shareddatawrite);
+		
+		VMCompiler->WriteCode(ResultIndex);//Set idx
+		VMCompiler->WriteCode(VarIndex);//Index of variable buffer in the set.
+		VMCompiler->WriteCode(IndexIndex);//Index of the temporary buffer that contains the data index.
+
+		VMCompiler->WriteCode(DataIndex);//Temp buffer index for the data to write.
+		return true;
+	}
+
+	FNiagaraDataSetID Event;
+};
+
+/** Expression that gets the next index in a shared data set and increments the counter. */ 
+class FNiagaraExpression_AcquireSharedDataIndex : public FNiagaraExpression
+{
+public:
+	FNiagaraExpression_AcquireSharedDataIndex(class FNiagaraCompiler* InCompiler, FNiagaraDataSetID InDataSet, bool bInWrap, TNiagaraExprPtr& ValidExpr)
+		: FNiagaraExpression(InCompiler, FNiagaraVariableInfo(TEXT("AcquireIndex"), ENiagaraDataType::Vector))
+		, DataSet(InDataSet)
+		, bWrap(bInWrap)
+	{
+		ResultLocation = ENiagaraExpressionResultLocation::Temporaries;
+		SourceExpressions.Add(ValidExpr);
+	}
+
+	~FNiagaraExpression_AcquireSharedDataIndex()
+	{
+		Compiler->FreeTemporary(ResultIndex);
+	}
+
+	virtual bool Process()override
+	{
+		FNiagaraCompiler_VectorVM* VMCompiler = (FNiagaraCompiler_VectorVM*)Compiler;
+		check(ResultLocation == ENiagaraExpressionResultLocation::Temporaries);
+		ResultIndex = VMCompiler->AquireTemporary();
+		if (ResultIndex == INDEX_NONE)
+		{
+			return false;
+		}
+
+		FNiagaraExpression* ValidExpr = GetSourceExpression(0).Get();
+		check(SourceExpressions.Num() == 1);
+
+		uint8 ValidIndex = VMCompiler->GetResultVMIndex(ValidExpr);
+
+		int32 SharedDataSetIndex;
+		VMCompiler->GetSharedDataIndex(DataSet, false, SharedDataSetIndex);
+		check(SharedDataSetIndex <= 255)
+		//TODO - Make this better, either ensure this is in a register by now or allow the VM to handle from constants?
+		//Though, why you'd have the conditional come from a constant, IDK.
+		if (!(ValidExpr->ResultLocation == ENiagaraExpressionResultLocation::Temporaries || ValidExpr->ResultLocation == ENiagaraExpressionResultLocation::InputData))
+		{
+			VMCompiler->Error(LOCTEXT("CompileFaliure_CannotUseConstantForSharedDataValidMask", "Use of constants for the shared data valid mask is not supported currently."), NULL, NULL);
+			return false;
+		}
+
+		VMCompiler->WriteCode(bWrap ? (uint8)EVectorVMOp::aquireshareddataindexwrap : (uint8)EVectorVMOp::aquireshareddataindex);
+		VMCompiler->WriteCode(ResultIndex);
+		VMCompiler->WriteCode(ValidIndex);
+		VMCompiler->WriteCode((uint8)SharedDataSetIndex);
+		return true;
+	}
+
+	FNiagaraDataSetID DataSet;
+	bool bWrap;
+};
+
+/** Expression that gets the next index in a shared data set and decrements the counter. */ 
+class FNiagaraExpression_ConsumeSharedDataIndex : public FNiagaraExpression
+{
+public:
+	FNiagaraExpression_ConsumeSharedDataIndex(class FNiagaraCompiler* InCompiler, FNiagaraDataSetID InDataSet, bool bInWrap)
+		: FNiagaraExpression(InCompiler, FNiagaraVariableInfo(TEXT("ConsumeIndex"), ENiagaraDataType::Vector))
+		, DataSet(InDataSet)
+		, bWrap(bInWrap)
+	{
+		ResultLocation = ENiagaraExpressionResultLocation::Temporaries;
+	}
+
+	~FNiagaraExpression_ConsumeSharedDataIndex()
+	{
+		Compiler->FreeTemporary(ResultIndex);
+	}
+
+	virtual bool Process()override
+	{
+		FNiagaraCompiler_VectorVM* VMCompiler = (FNiagaraCompiler_VectorVM*)Compiler;
+		ResultIndex = VMCompiler->AquireTemporary();
+		if (ResultIndex == INDEX_NONE)
+		{
+			return false;
+		}
+			
+		int32 SharedDataSetIndex;
+		VMCompiler->GetSharedDataIndex(DataSet, true, SharedDataSetIndex);
+			
+		VMCompiler->WriteCode(bWrap ? (uint8)EVectorVMOp::consumeshareddataindexwrap : (uint8)EVectorVMOp::consumeshareddataindex);
+		VMCompiler->WriteCode(ResultIndex);
+		VMCompiler->WriteCode((uint8)SharedDataSetIndex);
+		return true;
+	}
+
+	FNiagaraDataSetID DataSet;
+	bool bWrap;
+};
+
+//Can remove this if we refactor the compiler to allow a single expression to output multiple results.
+class FNiagaraExpression_IsValidSharedDataIndex : public FNiagaraExpression
+{
+public:
+	FNiagaraExpression_IsValidSharedDataIndex(class FNiagaraCompiler* InCompiler, FNiagaraDataSetID InDataSet, TNiagaraExprPtr& IndexExpr)
+		: FNiagaraExpression(InCompiler, FNiagaraVariableInfo(TEXT("IsValidIndex"), ENiagaraDataType::Vector))
+		, DataSet(InDataSet)
+	{
+		ResultLocation = ENiagaraExpressionResultLocation::Temporaries;
+		SourceExpressions.Add(IndexExpr);
+	}
+
+	~FNiagaraExpression_IsValidSharedDataIndex()
+	{
+		Compiler->FreeTemporary(ResultIndex);
+	}
+
+	virtual bool Process()override
+	{
+		FNiagaraCompiler_VectorVM* VMCompiler = (FNiagaraCompiler_VectorVM*)Compiler;
+		ResultIndex = VMCompiler->AquireTemporary();
+		if (ResultIndex == INDEX_NONE)
+		{
+			return false;
+		}
+		check(SourceExpressions.Num() == 1);
+
+		int32 IndexIndex = VMCompiler->GetResultVMIndex(GetSourceExpression(0).Get());
+
+		int32 SharedDataSetIndex;
+		FNiagaraDataSetProperties* DataSetProps = VMCompiler->GetSharedDataIndex(DataSet, true, SharedDataSetIndex);//This is a bit wonky. Might have to rework it a bit
+	
+		VMCompiler->WriteCode((uint8)EVectorVMOp::shareddataindexvalid);
+		VMCompiler->WriteCode(ResultIndex);
+		VMCompiler->WriteCode(IndexIndex);
+		VMCompiler->WriteCode(SharedDataSetIndex);
+		return true;
+	}
+
+	FNiagaraDataSetID DataSet;
+};
+
 bool FNiagaraCompiler_VectorVM::ExpressionIsConstant(FNiagaraExpression*  Expression)
 {
 	return Expression->ResultLocation == ENiagaraExpressionResultLocation::Constants;
@@ -129,11 +354,18 @@ bool FNiagaraCompiler_VectorVM::ExpressionIsBufferConstant(FNiagaraExpression*
 	return Expression->ResultLocation == ENiagaraExpressionResultLocation::BufferConstants;
 }
 
-VectorVM::EOperandType FNiagaraCompiler_VectorVM::GetVMOperandType(FNiagaraExpression *Expression)
+EVectorVMOperandLocation FNiagaraCompiler_VectorVM::GetVMOperandType(FNiagaraExpression *Expression)
 {
-	return Expression->ResultLocation == ENiagaraExpressionResultLocation::Constants ? VectorVM::ConstantOperandType :
-		Expression->ResultLocation == ENiagaraExpressionResultLocation::BufferConstants ? VectorVM::DataObjConstantOperandType :
-		VectorVM::RegisterOperandType;
+	switch (Expression->ResultLocation)
+	{
+	case ENiagaraExpressionResultLocation::Temporaries: return EVectorVMOperandLocation::TemporaryRegister;
+	case ENiagaraExpressionResultLocation::InputData: return EVectorVMOperandLocation::InputRegister;
+	case ENiagaraExpressionResultLocation::OutputData: return EVectorVMOperandLocation::OutputRegister;
+	case ENiagaraExpressionResultLocation::Constants: return EVectorVMOperandLocation::Constant;
+	case ENiagaraExpressionResultLocation::BufferConstants: return EVectorVMOperandLocation::DataObjConstant;
+	case ENiagaraExpressionResultLocation::SharedData: return EVectorVMOperandLocation::SharedData;
+	default: return EVectorVMOperandLocation::Undefined;
+	};
 }
 
 uint8 FNiagaraCompiler_VectorVM::GetResultVMIndex(FNiagaraExpression* Expression)
@@ -148,26 +380,31 @@ uint8 FNiagaraCompiler_VectorVM::GetResultVMIndex(FNiagaraExpression* Expression
 	}
 	else
 	{
-		if (Expression->ResultLocation == ENiagaraExpressionResultLocation::BufferConstants)
-		{
-			return Expression->ResultIndex; //TempRegister or Constant.
-		}
-
-		return Expression->ResultIndex;//TempRegister or Constant.
+		//TempRegisters start at 0 so can use index directly.
+		//Buffer constants, Constants and SharedData all exist in their own space so use index directly.
+		return Expression->ResultIndex;
 	}
 }
 
-ENiagaraDataType FNiagaraCompiler_VectorVM::GetConstantResultIndex(const FNiagaraVariableInfo& Constant, bool bInternal, int32& OutResultIndex, int32& OutComponentIndex)
+ENiagaraDataType FNiagaraCompiler_VectorVM::GetConstantResultIndex(const FNiagaraVariableInfo& Constant, bool bInternal, int32& OutResultIndex)
 {
 	ENiagaraDataType Type;
-	ConstantData.GetTableIndex(Constant, bInternal, OutResultIndex, OutComponentIndex, Type);
+	ConstantData.GetTableIndex(Constant, bInternal, OutResultIndex, Type);
 	return Type;
 }
 
 int32 FNiagaraCompiler_VectorVM::AquireTemporary()
 {
 	int32 Free = TempRegisters.Find(false);
-	TempRegisters[Free] = true;
+	if (Free != INDEX_NONE)
+	{
+		TempRegisters[Free] = true;
+	}
+	else
+	{
+		Error(LOCTEXT("ScriptUsesTooManyTemporariesError","This script uses too many temporary registers. It has to be simplified to work."), NULL, NULL);
+	}
+
 	return Free;
 }
 
@@ -189,6 +426,74 @@ TNiagaraExprPtr FNiagaraCompiler_VectorVM::Output(const FNiagaraVariableInfo& At
 	return Expressions[Index];
 }
 
+TNiagaraExprPtr FNiagaraCompiler_VectorVM::SharedDataRead(const FNiagaraDataSetID& DataSet, const FNiagaraVariableInfo& Variable, TNiagaraExprPtr& IndexExpr)
+{
+	bool bFound = false;
+	check(DataSet.Type == ENiagaraDataSetType::Event);//TODO: Other data set types.
+	for (FNiagaraDataSetProperties& SetProps : EventReceivers)
+	{
+		if (SetProps.ID == DataSet)
+		{
+			SetProps.Variables.AddUnique(Variable);
+			bFound = true;
+			break;
+		}
+	}
+	if (!bFound)
+	{
+		int32 NewIdx = EventReceivers.AddDefaulted();
+		EventReceivers[NewIdx].ID = DataSet;
+		EventReceivers[NewIdx].Variables.Add(Variable);
+	}
+	
+	//TODO - This can be moved up into FNiagaraCompiler when I work around the issue of moving the shared data into a temporary.
+	int32 Index = Expressions.Add(MakeShareable(new FNiagaraExpression_VMDataSetRead(this, DataSet, Variable, IndexExpr)));
+	return Expressions[Index];
+}
+
+TNiagaraExprPtr FNiagaraCompiler_VectorVM::SharedDataWrite(const FNiagaraDataSetID& DataSet, const FNiagaraVariableInfo& Variable, TNiagaraExprPtr& IndexExpr, TNiagaraExprPtr& DataExpr)
+{
+	bool bFound = false;
+	check(DataSet.Type == ENiagaraDataSetType::Event);//TODO: Other data set types.
+	for (FNiagaraDataSetProperties& SetProps : EventGenerators)
+	{
+		if (SetProps.ID == DataSet)
+		{
+			SetProps.Variables.AddUnique(Variable);
+			bFound = true;
+			break;
+		}
+	}
+	if (!bFound)
+	{
+		int32 NewIdx = EventGenerators.AddDefaulted();
+		EventGenerators[NewIdx].ID = DataSet;
+		EventGenerators[NewIdx].Variables.Add(Variable);
+	}
+
+	//TODO - This can be moved up into FNiagaraCompiler when I work around the issue of moving the shared data into a temporary.
+	int32 Index = Expressions.Add(MakeShareable(new FNiagaraExpression_VMDataSetWrite(this, DataSet, Variable, IndexExpr, DataExpr)));
+	return Expressions[Index];
+}
+
+TNiagaraExprPtr FNiagaraCompiler_VectorVM::AcquireSharedDataIndex(const FNiagaraDataSetID& DataSet, bool bWrap, TNiagaraExprPtr& ValidExpr)
+{
+	int32 Index = Expressions.Add(MakeShareable(new FNiagaraExpression_AcquireSharedDataIndex(this, DataSet, bWrap, ValidExpr)));
+	return Expressions[Index];
+}
+
+TNiagaraExprPtr FNiagaraCompiler_VectorVM::ConsumeSharedDataIndex(const FNiagaraDataSetID& DataSet, bool bWrap)
+{
+	int32 Index = Expressions.Add(MakeShareable(new FNiagaraExpression_ConsumeSharedDataIndex(this, DataSet, bWrap)));
+	return Expressions[Index];
+}
+
+TNiagaraExprPtr FNiagaraCompiler_VectorVM::SharedDataIndexIsValid(const FNiagaraDataSetID& DataSet, TNiagaraExprPtr& IndexExpr)
+{
+	int32 Index = Expressions.Add(MakeShareable(new FNiagaraExpression_IsValidSharedDataIndex(this, DataSet, IndexExpr)));
+	return Expressions[Index];
+}
+
 bool FNiagaraCompiler_VectorVM::CompileScript(UNiagaraScript* InScript)
 {
 	//TODO - Most of this function can be refactored up into FNiagaraCompiler once we know how the compute script and VM script will coexist in the UNiagaraScript.
@@ -196,10 +501,10 @@ bool FNiagaraCompiler_VectorVM::CompileScript(UNiagaraScript* InScript)
 	check(InScript);
 
 	Script = InScript;
-	Source = CastChecked<UNiagaraScriptSource>(InScript->Source);
+	Source = CastChecked<UNiagaraScriptSource>(Script->Source);
 
 	//Clear previous graph errors.
-	if (InScript->ByteCode.Num() == 0)
+	if (Script->ByteCode.Num() == 0)
 	{
 		bool bChanged = false;
 		for (UEdGraphNode* Node : Source->NodeGraph->Nodes)
@@ -226,85 +531,111 @@ bool FNiagaraCompiler_VectorVM::CompileScript(UNiagaraScript* InScript)
 	SourceGraph = CastChecked<UNiagaraGraph>(FEdGraphUtilities::CloneGraph(Source->NodeGraph, NULL, &MessageLog));
 	FEdGraphUtilities::MergeChildrenGraphsIn(SourceGraph, SourceGraph, /*bRequireSchemaMatch=*/ true);
 
-	if (MergeInFunctionNodes())
-	{
-		TempRegisters.AddZeroed(VectorVM::NumTempRegisters);
 
-		// Find the output node.
-		UNiagaraNodeOutput* OutputNode = SourceGraph->FindOutputNode();
+	TempRegisters.AddZeroed(VectorVM::NumTempRegisters);
 
-		TArray<FNiagaraNodeResult> OutputExpressions;
-		OutputNode->Compile(this, OutputExpressions);
+	// Find the output node and compile it.
+	UNiagaraNodeOutput* OutputNode = SourceGraph->FindOutputNode();
+	TArray<FNiagaraNodeResult> OutputExpressions;
+	OutputNode->Compile(this, OutputExpressions);
 
-		//Todo - track usage of temp registers so we can allocate and free and reuse from a smaller set.
-		int32 NextTempRegister = 0;
-		Script->ByteCode.Empty();
+	//Find any write data set nodes and compile them.
+	TArray<UNiagaraNodeWriteDataSet*> WriteDataSetNodes;
+	SourceGraph->FindWriteDataSetNodes(WriteDataSetNodes);
+	for (UNiagaraNodeWriteDataSet* WriteNode : WriteDataSetNodes)
+	{
+		WriteNode->Compile(this, OutputExpressions);
+	}
+
+	//Clear the pin to expressions map. All reused expressions will now be referenced by the expressions using them.
+	PinToExpression.Empty();//TODO - Replace this with expression caching at a lower level. Will be more useful and nicer.
 
-		//Now we have a set of expressions from which we can generate the bytecode.
-		for (int32 i = 0; i < Expressions.Num(); ++i)
+	//Todo - track usage of temp registers so we can allocate and free and reuse from a smaller set.
+	int32 NextTempRegister = 0;
+	Script->ByteCode.Empty();
+
+	//Now we have a set of expressions from which we can generate the bytecode.
+	for (int32 i = 0; i < Expressions.Num(); ++i)
+	{
+		TNiagaraExprPtr Expr = Expressions[i];
+		if (Expr.IsValid())
 		{
-			TNiagaraExprPtr Expr = Expressions[i];
-			if (Expr.IsValid())
-			{
-				Expr->Process();
-				Expr->PostProcess();
-			}
-			else
+			if (!Expr->Process())
 			{
-				check(0);//Replace with graceful bailout?
+				Error(LOCTEXT("CompilationError", "Compilation Error"), NULL, NULL);
+				break;
 			}
 
-			Expressions[i].Reset();
+			Expr->PostProcess();
 		}
-
-		PinToExpression.Empty();//TODO - Replace this with expression caching at a lower level. Will be more useful and nicer.
-
-		if (MessageLog.NumErrors == 0)
+		else
 		{
-			// Terminate with the 'done' opcode.
-			Script->ByteCode.Add((int8)VectorVM::EOp::done);
-			Script->ConstantData = ConstantData;
-			Script->Attributes = OutputNode->Outputs;
-			return true;
+			check(0);//Replace with graceful bailout?
 		}
+
+		Expressions[i].Reset();
+	}
+
+	Expressions.Empty();
+
+	if (MessageLog.NumErrors == 0)
+	{
+		check(TempRegisters.Find(true) == INDEX_NONE);//Make sure there aren't any temp registers still in use. If so, some expression is leaking them.
+
+		Script->Modify();
+		Source->FlattenedNodeGraph = SourceGraph;
+		Source->FlattenedNodeGraph->Modify();			
+		Source->FlattenedNodeGraph->NotifyGraphChanged();
+
+		// Terminate with the 'done' opcode.
+		Script->ByteCode.Add((int8)EVectorVMOp::done);
+		Script->ConstantData = ConstantData;
+		Script->Attributes = OutputNode->Outputs;
+		Script->EventGenerators = EventGenerators;
+		Script->EventReceivers = EventReceivers;
+		Script->Usage = Usage;
+		return true;
 	}
 
 	//Compilation has failed.
-	InScript->Attributes.Empty();
-	InScript->ByteCode.Empty();
-	InScript->ConstantData.Empty();
+	Script->Attributes.Empty();
+	Script->ByteCode.Empty();
+	Script->ConstantData.Empty();
+	Script->EventGenerators.Empty();
+	Script->EventReceivers.Empty();
+	Script->Usage = FNiagaraScriptUsageInfo();
 	//Tell the graph to show the errors.
 	Source->NodeGraph->NotifyGraphChanged();
 		
 	return false;
 }
 
-TNiagaraExprPtr FNiagaraCompiler_VectorVM::Expression_VMNative(VectorVM::EOp Op, TArray<TNiagaraExprPtr>& InputExpressions)
+TNiagaraExprPtr FNiagaraCompiler_VectorVM::Expression_VMNative(EVectorVMOp Op, TArray<TNiagaraExprPtr>& InputExpressions)
 {
 	static const FName OpName(TEXT("VMOp"));
 	int32 Index = Expressions.Add(MakeShareable(new FNiagaraExpression_VMOperation(this, FNiagaraVariableInfo(OpName, ENiagaraDataType::Vector), InputExpressions, Op)));
 	return Expressions[Index];
 }
 
-TNiagaraExprPtr FNiagaraCompiler_VectorVM::Expression_VMNative(VectorVM::EOp Op, TNiagaraExprPtr A)
+TNiagaraExprPtr FNiagaraCompiler_VectorVM::Expression_VMNative(EVectorVMOp Op, TNiagaraExprPtr A)
 {
 	TArray<TNiagaraExprPtr> In; In.Add(A);
 	return Expression_VMNative(Op, In);
 }
 
-TNiagaraExprPtr FNiagaraCompiler_VectorVM::Expression_VMNative(VectorVM::EOp Op, TNiagaraExprPtr A, TNiagaraExprPtr B)
+TNiagaraExprPtr FNiagaraCompiler_VectorVM::Expression_VMNative(EVectorVMOp Op, TNiagaraExprPtr A, TNiagaraExprPtr B)
 {
 	TArray<TNiagaraExprPtr> In; In.Add(A); In.Add(B);
 	return Expression_VMNative(Op, In);
 }
 
-TNiagaraExprPtr FNiagaraCompiler_VectorVM::Expression_VMNative(VectorVM::EOp Op, TNiagaraExprPtr A, TNiagaraExprPtr B, TNiagaraExprPtr C)
+TNiagaraExprPtr FNiagaraCompiler_VectorVM::Expression_VMNative(EVectorVMOp Op, TNiagaraExprPtr A, TNiagaraExprPtr B, TNiagaraExprPtr C)
 {
 	TArray<TNiagaraExprPtr> In; In.Add(A); In.Add(B); In.Add(C);
 	return Expression_VMNative(Op, In);
 }
 
-TNiagaraExprPtr FNiagaraCompiler_VectorVM::Expression_VMNative(VectorVM::EOp Op, TNiagaraExprPtr A, TNiagaraExprPtr B, TNiagaraExprPtr C, TNiagaraExprPtr D)
+TNiagaraExprPtr FNiagaraCompiler_VectorVM::Expression_VMNative(EVectorVMOp Op, TNiagaraExprPtr A, TNiagaraExprPtr B, TNiagaraExprPtr C, TNiagaraExprPtr D)
 {
 	TArray<TNiagaraExprPtr> In; In.Add(A); In.Add(B); In.Add(C); In.Add(D);
 	return Expression_VMNative(Op, In);
@@ -312,235 +643,235 @@ TNiagaraExprPtr FNiagaraCompiler_VectorVM::Expression_VMNative(VectorVM::EOp Op,
 
 bool FNiagaraCompiler_VectorVM::Add_Internal(TArray<TNiagaraExprPtr>& InputExpressions, TArray<TNiagaraExprPtr>& OutputExpressions)
 {
-	OutputExpressions.Add(Expression_VMNative(VectorVM::EOp::add, InputExpressions));
+	OutputExpressions.Add(Expression_VMNative(EVectorVMOp::add, InputExpressions));
 	return true;
 }
 
 bool FNiagaraCompiler_VectorVM::Subtract_Internal(TArray<TNiagaraExprPtr>& InputExpressions, TArray<TNiagaraExprPtr>& OutputExpressions)
 {
-	OutputExpressions.Add((Expression_VMNative(VectorVM::EOp::sub, InputExpressions)));
+	OutputExpressions.Add((Expression_VMNative(EVectorVMOp::sub, InputExpressions)));
 	return true;
 }
 
 bool FNiagaraCompiler_VectorVM::Multiply_Internal(TArray<TNiagaraExprPtr>& InputExpressions, TArray<TNiagaraExprPtr>& OutputExpressions)
 {
-	OutputExpressions.Add(Expression_VMNative(VectorVM::EOp::mul, InputExpressions));
+	OutputExpressions.Add(Expression_VMNative(EVectorVMOp::mul, InputExpressions));
 	return true;
 }
 
 bool FNiagaraCompiler_VectorVM::Divide_Internal(TArray<TNiagaraExprPtr>& InputExpressions, TArray<TNiagaraExprPtr>& OutputExpressions)
 {
-	OutputExpressions.Add(Expression_VMNative(VectorVM::EOp::div, InputExpressions));
+	OutputExpressions.Add(Expression_VMNative(EVectorVMOp::div, InputExpressions));
 	return true;
 }
 
 bool FNiagaraCompiler_VectorVM::MultiplyAdd_Internal(TArray<TNiagaraExprPtr>& InputExpressions, TArray<TNiagaraExprPtr>& OutputExpressions)
 {
-	OutputExpressions.Add(Expression_VMNative(VectorVM::EOp::mad, InputExpressions));
+	OutputExpressions.Add(Expression_VMNative(EVectorVMOp::mad, InputExpressions));
 	return true;
 }
 
 bool FNiagaraCompiler_VectorVM::Lerp_Internal(TArray<TNiagaraExprPtr>& InputExpressions, TArray<TNiagaraExprPtr>& OutputExpressions)
 {
-	OutputExpressions.Add(Expression_VMNative(VectorVM::EOp::lerp, InputExpressions));
+	OutputExpressions.Add(Expression_VMNative(EVectorVMOp::lerp, InputExpressions));
 	return true;
 }
 
 bool FNiagaraCompiler_VectorVM::Reciprocal_Internal(TArray<TNiagaraExprPtr>& InputExpressions, TArray<TNiagaraExprPtr>& OutputExpressions)
 {
-	OutputExpressions.Add(Expression_VMNative(VectorVM::EOp::rcp, InputExpressions));
+	OutputExpressions.Add(Expression_VMNative(EVectorVMOp::rcp, InputExpressions));
 	return true;
 }
 
 bool FNiagaraCompiler_VectorVM::ReciprocalSqrt_Internal(TArray<TNiagaraExprPtr>& InputExpressions, TArray<TNiagaraExprPtr>& OutputExpressions)
 {
-	OutputExpressions.Add(Expression_VMNative(VectorVM::EOp::rsq, InputExpressions));
+	OutputExpressions.Add(Expression_VMNative(EVectorVMOp::rsq, InputExpressions));
 	return true;
 }
 
 bool FNiagaraCompiler_VectorVM::Sqrt_Internal(TArray<TNiagaraExprPtr>& InputExpressions, TArray<TNiagaraExprPtr>& OutputExpressions)
 {
-	OutputExpressions.Add(Expression_VMNative(VectorVM::EOp::sqrt, InputExpressions));
+	OutputExpressions.Add(Expression_VMNative(EVectorVMOp::sqrt, InputExpressions));
 	return true;
 }
 
 bool FNiagaraCompiler_VectorVM::Negate_Internal(TArray<TNiagaraExprPtr>& InputExpressions, TArray<TNiagaraExprPtr>& OutputExpressions)
 {
-	OutputExpressions.Add(Expression_VMNative(VectorVM::EOp::neg, InputExpressions));
+	OutputExpressions.Add(Expression_VMNative(EVectorVMOp::neg, InputExpressions));
 	return true;
 }
 
 bool FNiagaraCompiler_VectorVM::Abs_Internal(TArray<TNiagaraExprPtr>& InputExpressions, TArray<TNiagaraExprPtr>& OutputExpressions)
 {
-	OutputExpressions.Add(Expression_VMNative(VectorVM::EOp::abs, InputExpressions));
+	OutputExpressions.Add(Expression_VMNative(EVectorVMOp::abs, InputExpressions));
 	return true;
 }
 
 bool FNiagaraCompiler_VectorVM::Exp_Internal(TArray<TNiagaraExprPtr>& InputExpressions, TArray<TNiagaraExprPtr>& OutputExpressions)
 {
-	OutputExpressions.Add(Expression_VMNative(VectorVM::EOp::exp, InputExpressions));
+	OutputExpressions.Add(Expression_VMNative(EVectorVMOp::exp, InputExpressions));
 	return true;
 }
 
 bool FNiagaraCompiler_VectorVM::Exp2_Internal(TArray<TNiagaraExprPtr>& InputExpressions, TArray<TNiagaraExprPtr>& OutputExpressions)
 {
-	OutputExpressions.Add(Expression_VMNative(VectorVM::EOp::exp2, InputExpressions));
+	OutputExpressions.Add(Expression_VMNative(EVectorVMOp::exp2, InputExpressions));
 	return true;
 }
 
 bool FNiagaraCompiler_VectorVM::Log_Internal(TArray<TNiagaraExprPtr>& InputExpressions, TArray<TNiagaraExprPtr>& OutputExpressions)
 {
-	OutputExpressions.Add(Expression_VMNative(VectorVM::EOp::log, InputExpressions));
+	OutputExpressions.Add(Expression_VMNative(EVectorVMOp::log, InputExpressions));
 	return true;
 }
 
 bool FNiagaraCompiler_VectorVM::LogBase2_Internal(TArray<TNiagaraExprPtr>& InputExpressions, TArray<TNiagaraExprPtr>& OutputExpressions)
 {
-	OutputExpressions.Add(Expression_VMNative(VectorVM::EOp::log2, InputExpressions));
+	OutputExpressions.Add(Expression_VMNative(EVectorVMOp::log2, InputExpressions));
 	return true;
 }
 
 bool FNiagaraCompiler_VectorVM::Sin_Internal(TArray<TNiagaraExprPtr>& InputExpressions, TArray<TNiagaraExprPtr>& OutputExpressions)
 {
-	OutputExpressions.Add(Expression_VMNative(VectorVM::EOp::sin, InputExpressions));
+	OutputExpressions.Add(Expression_VMNative(EVectorVMOp::sin, InputExpressions));
 	return true;
 }
 
 bool FNiagaraCompiler_VectorVM::Cos_Internal(TArray<TNiagaraExprPtr>& InputExpressions, TArray<TNiagaraExprPtr>& OutputExpressions)
 {
-	OutputExpressions.Add(Expression_VMNative(VectorVM::EOp::cos, InputExpressions));
+	OutputExpressions.Add(Expression_VMNative(EVectorVMOp::cos, InputExpressions));
 	return true;
 }
 
 bool FNiagaraCompiler_VectorVM::Tan_Internal(TArray<TNiagaraExprPtr>& InputExpressions, TArray<TNiagaraExprPtr>& OutputExpressions)
 {
-	OutputExpressions.Add(Expression_VMNative(VectorVM::EOp::tan, InputExpressions));
+	OutputExpressions.Add(Expression_VMNative(EVectorVMOp::tan, InputExpressions));
 	return true;
 }
 
 bool FNiagaraCompiler_VectorVM::ASin_Internal(TArray<TNiagaraExprPtr>& InputExpressions, TArray<TNiagaraExprPtr>& OutputExpressions)
 {
-	OutputExpressions.Add(Expression_VMNative(VectorVM::EOp::asin, InputExpressions));
+	OutputExpressions.Add(Expression_VMNative(EVectorVMOp::asin, InputExpressions));
 	return true;
 }
 
 bool FNiagaraCompiler_VectorVM::ACos_Internal(TArray<TNiagaraExprPtr>& InputExpressions, TArray<TNiagaraExprPtr>& OutputExpressions)
 {
-	OutputExpressions.Add(Expression_VMNative(VectorVM::EOp::acos, InputExpressions));
+	OutputExpressions.Add(Expression_VMNative(EVectorVMOp::acos, InputExpressions));
 	return true;
 }
 
 bool FNiagaraCompiler_VectorVM::ATan_Internal(TArray<TNiagaraExprPtr>& InputExpressions, TArray<TNiagaraExprPtr>& OutputExpressions)
 {
-	OutputExpressions.Add(Expression_VMNative(VectorVM::EOp::atan, InputExpressions));
+	OutputExpressions.Add(Expression_VMNative(EVectorVMOp::atan, InputExpressions));
 	return true;
 }
 
 bool FNiagaraCompiler_VectorVM::ATan2_Internal(TArray<TNiagaraExprPtr>& InputExpressions, TArray<TNiagaraExprPtr>& OutputExpressions)
 {
-	OutputExpressions.Add(Expression_VMNative(VectorVM::EOp::atan2, InputExpressions));
+	OutputExpressions.Add(Expression_VMNative(EVectorVMOp::atan2, InputExpressions));
 	return true;
 }
 
 bool FNiagaraCompiler_VectorVM::Ceil_Internal(TArray<TNiagaraExprPtr>& InputExpressions, TArray<TNiagaraExprPtr>& OutputExpressions)
 {
-	OutputExpressions.Add(Expression_VMNative(VectorVM::EOp::ceil, InputExpressions));
+	OutputExpressions.Add(Expression_VMNative(EVectorVMOp::ceil, InputExpressions));
 	return true;
 }
 
 bool FNiagaraCompiler_VectorVM::Floor_Internal(TArray<TNiagaraExprPtr>& InputExpressions, TArray<TNiagaraExprPtr>& OutputExpressions)
 {
-	OutputExpressions.Add(Expression_VMNative(VectorVM::EOp::floor, InputExpressions));
+	OutputExpressions.Add(Expression_VMNative(EVectorVMOp::floor, InputExpressions));
 	return true;
 }
 
 bool FNiagaraCompiler_VectorVM::Modulo_Internal(TArray<TNiagaraExprPtr>& InputExpressions, TArray<TNiagaraExprPtr>& OutputExpressions)
 {
-	OutputExpressions.Add(Expression_VMNative(VectorVM::EOp::fmod, InputExpressions));
+	OutputExpressions.Add(Expression_VMNative(EVectorVMOp::fmod, InputExpressions));
 	return true;
 }
 
 bool FNiagaraCompiler_VectorVM::Fractional_Internal(TArray<TNiagaraExprPtr>& InputExpressions, TArray<TNiagaraExprPtr>& OutputExpressions)
 {
-	OutputExpressions.Add(Expression_VMNative(VectorVM::EOp::frac, InputExpressions));
+	OutputExpressions.Add(Expression_VMNative(EVectorVMOp::frac, InputExpressions));
 	return true;
 }
 
 bool FNiagaraCompiler_VectorVM::Truncate_Internal(TArray<TNiagaraExprPtr>& InputExpressions, TArray<TNiagaraExprPtr>& OutputExpressions)
 {
-	OutputExpressions.Add(Expression_VMNative(VectorVM::EOp::trunc, InputExpressions));
+	OutputExpressions.Add(Expression_VMNative(EVectorVMOp::trunc, InputExpressions));
 	return true;
 }
 
 bool FNiagaraCompiler_VectorVM::Clamp_Internal(TArray<TNiagaraExprPtr>& InputExpressions, TArray<TNiagaraExprPtr>& OutputExpressions)
 {
-	OutputExpressions.Add(Expression_VMNative(VectorVM::EOp::clamp, InputExpressions));
+	OutputExpressions.Add(Expression_VMNative(EVectorVMOp::clamp, InputExpressions));
 	return true;
 }
 
 bool FNiagaraCompiler_VectorVM::Min_Internal(TArray<TNiagaraExprPtr>& InputExpressions, TArray<TNiagaraExprPtr>& OutputExpressions)
 {
-	OutputExpressions.Add(Expression_VMNative(VectorVM::EOp::min, InputExpressions));
+	OutputExpressions.Add(Expression_VMNative(EVectorVMOp::min, InputExpressions));
 	return true;
 }
 
 bool FNiagaraCompiler_VectorVM::Max_Internal(TArray<TNiagaraExprPtr>& InputExpressions, TArray<TNiagaraExprPtr>& OutputExpressions)
 {
-	OutputExpressions.Add(Expression_VMNative(VectorVM::EOp::max, InputExpressions));
+	OutputExpressions.Add(Expression_VMNative(EVectorVMOp::max, InputExpressions));
 	return true;
 }
 
 bool FNiagaraCompiler_VectorVM::Pow_Internal(TArray<TNiagaraExprPtr>& InputExpressions, TArray<TNiagaraExprPtr>& OutputExpressions)
 {
-	OutputExpressions.Add(Expression_VMNative(VectorVM::EOp::pow, InputExpressions));
+	OutputExpressions.Add(Expression_VMNative(EVectorVMOp::pow, InputExpressions));
 	return true;
 }
 
 bool FNiagaraCompiler_VectorVM::Sign_Internal(TArray<TNiagaraExprPtr>& InputExpressions, TArray<TNiagaraExprPtr>& OutputExpressions)
 {
-	OutputExpressions.Add(Expression_VMNative(VectorVM::EOp::sign, InputExpressions));
+	OutputExpressions.Add(Expression_VMNative(EVectorVMOp::sign, InputExpressions));
 	return true;
 }
 
 bool FNiagaraCompiler_VectorVM::Step_Internal(TArray<TNiagaraExprPtr>& InputExpressions, TArray<TNiagaraExprPtr>& OutputExpressions)
 {
-	OutputExpressions.Add(Expression_VMNative(VectorVM::EOp::step, InputExpressions));
+	OutputExpressions.Add(Expression_VMNative(EVectorVMOp::step, InputExpressions));
 	return true;
 }
 
 bool FNiagaraCompiler_VectorVM::Dot_Internal(TArray<TNiagaraExprPtr>& InputExpressions, TArray<TNiagaraExprPtr>& OutputExpressions)
 {
-	OutputExpressions.Add(Expression_VMNative(VectorVM::EOp::dot, InputExpressions));
+	OutputExpressions.Add(Expression_VMNative(EVectorVMOp::dot, InputExpressions));
 	return true;
 }
 
 bool FNiagaraCompiler_VectorVM::Cross_Internal(TArray<TNiagaraExprPtr>& InputExpressions, TArray<TNiagaraExprPtr>& OutputExpressions)
 {
-	OutputExpressions.Add(Expression_VMNative(VectorVM::EOp::cross, InputExpressions));
+	OutputExpressions.Add(Expression_VMNative(EVectorVMOp::cross, InputExpressions));
 	return true;
 }
 
 bool FNiagaraCompiler_VectorVM::Normalize_Internal(TArray<TNiagaraExprPtr>& InputExpressions, TArray<TNiagaraExprPtr>& OutputExpressions)
 {
-	OutputExpressions.Add(Expression_VMNative(VectorVM::EOp::normalize, InputExpressions));
+	OutputExpressions.Add(Expression_VMNative(EVectorVMOp::normalize, InputExpressions));
 	return true;
 }
 
 bool FNiagaraCompiler_VectorVM::Random_Internal(TArray<TNiagaraExprPtr>& InputExpressions, TArray<TNiagaraExprPtr>& OutputExpressions)
 {
-	OutputExpressions.Add(Expression_VMNative(VectorVM::EOp::random, InputExpressions));
+	OutputExpressions.Add(Expression_VMNative(EVectorVMOp::random, InputExpressions));
 	return true;
 }
 
 bool FNiagaraCompiler_VectorVM::Length_Internal(TArray<TNiagaraExprPtr>& InputExpressions, TArray<TNiagaraExprPtr>& OutputExpressions)
 {
-	OutputExpressions.Add(Expression_VMNative(VectorVM::EOp::length, InputExpressions));
+	OutputExpressions.Add(Expression_VMNative(EVectorVMOp::length, InputExpressions));
 	return true;
 }
 
 bool FNiagaraCompiler_VectorVM::Noise_Internal(TArray<TNiagaraExprPtr>& InputExpressions, TArray<TNiagaraExprPtr>& OutputExpressions)
 {
-	OutputExpressions.Add(Expression_VMNative(VectorVM::EOp::noise, InputExpressions));
+	OutputExpressions.Add(Expression_VMNative(EVectorVMOp::noise, InputExpressions));
 	return true;
 }
 
@@ -567,31 +898,31 @@ bool FNiagaraCompiler_VectorVM::Matrix_Internal(TArray<TNiagaraExprPtr>& InputEx
 
 bool FNiagaraCompiler_VectorVM::Compose_Internal(TArray<TNiagaraExprPtr>& InputExpressions, TArray<TNiagaraExprPtr>& OutputExpressions)
 {
-	OutputExpressions.Add(Expression_VMNative(VectorVM::EOp::compose, InputExpressions));
+	OutputExpressions.Add(Expression_VMNative(EVectorVMOp::compose, InputExpressions));
 	return true;
 }
 
 bool FNiagaraCompiler_VectorVM::ComposeX_Internal(TArray<TNiagaraExprPtr>& InputExpressions, TArray<TNiagaraExprPtr>& OutputExpressions)
 {
-	OutputExpressions.Add(Expression_VMNative(VectorVM::EOp::composex, InputExpressions));
+	OutputExpressions.Add(Expression_VMNative(EVectorVMOp::composex, InputExpressions));
 	return true;
 }
 
 bool FNiagaraCompiler_VectorVM::ComposeY_Internal(TArray<TNiagaraExprPtr>& InputExpressions, TArray<TNiagaraExprPtr>& OutputExpressions)
 {
-	OutputExpressions.Add(Expression_VMNative(VectorVM::EOp::composey, InputExpressions));
+	OutputExpressions.Add(Expression_VMNative(EVectorVMOp::composey, InputExpressions));
 	return true;
 }
 
 bool FNiagaraCompiler_VectorVM::ComposeZ_Internal(TArray<TNiagaraExprPtr>& InputExpressions, TArray<TNiagaraExprPtr>& OutputExpressions)
 {
-	OutputExpressions.Add(Expression_VMNative(VectorVM::EOp::composez, InputExpressions));
+	OutputExpressions.Add(Expression_VMNative(EVectorVMOp::composez, InputExpressions));
 	return true;
 }
 
 bool FNiagaraCompiler_VectorVM::ComposeW_Internal(TArray<TNiagaraExprPtr>& InputExpressions, TArray<TNiagaraExprPtr>& OutputExpressions)
 {
-	OutputExpressions.Add(Expression_VMNative(VectorVM::EOp::composew, InputExpressions));
+	OutputExpressions.Add(Expression_VMNative(EVectorVMOp::composew, InputExpressions));
 	return true;
 }
 
@@ -606,25 +937,25 @@ bool FNiagaraCompiler_VectorVM::Split_Internal(TArray<TNiagaraExprPtr>& InputExp
 
 bool FNiagaraCompiler_VectorVM::GetX_Internal(TArray<TNiagaraExprPtr>& InputExpressions, TArray<TNiagaraExprPtr>& OutputExpressions)
 {
-	OutputExpressions.Add(Expression_VMNative(VectorVM::EOp::splatx, InputExpressions));
+	OutputExpressions.Add(Expression_VMNative(EVectorVMOp::splatx, InputExpressions));
 	return true;
 }
 
 bool FNiagaraCompiler_VectorVM::GetY_Internal(TArray<TNiagaraExprPtr>& InputExpressions, TArray<TNiagaraExprPtr>& OutputExpressions)
 {
-	OutputExpressions.Add(Expression_VMNative(VectorVM::EOp::splaty, InputExpressions));
+	OutputExpressions.Add(Expression_VMNative(EVectorVMOp::splaty, InputExpressions));
 	return true;
 }
 
 bool FNiagaraCompiler_VectorVM::GetZ_Internal(TArray<TNiagaraExprPtr>& InputExpressions, TArray<TNiagaraExprPtr>& OutputExpressions)
 {
-	OutputExpressions.Add(Expression_VMNative(VectorVM::EOp::splatz, InputExpressions));
+	OutputExpressions.Add(Expression_VMNative(EVectorVMOp::splatz, InputExpressions));
 	return true;
 }
 
 bool FNiagaraCompiler_VectorVM::GetW_Internal(TArray<TNiagaraExprPtr>& InputExpressions, TArray<TNiagaraExprPtr>& OutputExpressions)
 {
-	OutputExpressions.Add(Expression_VMNative(VectorVM::EOp::splatw, InputExpressions));
+	OutputExpressions.Add(Expression_VMNative(EVectorVMOp::splatw, InputExpressions));
 	return true;
 }
 
@@ -640,19 +971,19 @@ bool FNiagaraCompiler_VectorVM::TransformVector_Internal(TArray<TNiagaraExprPtr>
 	TNiagaraExprPtr VectorExpr = InputExpressions[1];
 
 	TArray<TNiagaraExprPtr> Temp;
-	Temp.Add(Expression_VMNative(VectorVM::EOp::splatx, VectorExpr));
-	Temp.Add(Expression_VMNative(VectorVM::EOp::splaty, VectorExpr));
-	Temp.Add(Expression_VMNative(VectorVM::EOp::splatz, VectorExpr));
-	Temp.Add(Expression_VMNative(VectorVM::EOp::splatw, VectorExpr));
-
-	Temp[0] = Expression_VMNative(VectorVM::EOp::mul, MatrixExpr->GetSourceExpression(0), Temp[0]);
-	Temp[1] = Expression_VMNative(VectorVM::EOp::mul, MatrixExpr->GetSourceExpression(1), Temp[1]);
-	Temp[2] = Expression_VMNative(VectorVM::EOp::mul, MatrixExpr->GetSourceExpression(2), Temp[2]);
-	Temp[3] = Expression_VMNative(VectorVM::EOp::mul, MatrixExpr->GetSourceExpression(3), Temp[3]);
-
-	Temp[0] = Expression_VMNative(VectorVM::EOp::add, Temp[0], Temp[1]);
-	Temp[1] = Expression_VMNative(VectorVM::EOp::add, Temp[2], Temp[3]);
-	OutputExpressions.Add(Expression_VMNative(VectorVM::EOp::add, Temp[0], Temp[1]));
+	Temp.Add(Expression_VMNative(EVectorVMOp::splatx, VectorExpr));
+	Temp.Add(Expression_VMNative(EVectorVMOp::splaty, VectorExpr));
+	Temp.Add(Expression_VMNative(EVectorVMOp::splatz, VectorExpr));
+	Temp.Add(Expression_VMNative(EVectorVMOp::splatw, VectorExpr));
+
+	Temp[0] = Expression_VMNative(EVectorVMOp::mul, MatrixExpr->GetSourceExpression(0), Temp[0]);
+	Temp[1] = Expression_VMNative(EVectorVMOp::mul, MatrixExpr->GetSourceExpression(1), Temp[1]);
+	Temp[2] = Expression_VMNative(EVectorVMOp::mul, MatrixExpr->GetSourceExpression(2), Temp[2]);
+	Temp[3] = Expression_VMNative(EVectorVMOp::mul, MatrixExpr->GetSourceExpression(3), Temp[3]);
+
+	Temp[0] = Expression_VMNative(EVectorVMOp::add, Temp[0], Temp[1]);
+	Temp[1] = Expression_VMNative(EVectorVMOp::add, Temp[2], Temp[3]);
+	OutputExpressions.Add(Expression_VMNative(EVectorVMOp::add, Temp[0], Temp[1]));
 	return true;
 }
 
@@ -668,10 +999,10 @@ bool FNiagaraCompiler_VectorVM::Transpose_Internal(TArray<TNiagaraExprPtr>& Inpu
 	//TODO: can probably make this faster if I provide direct shuffle ops in the VM rather than using the flexible, catch-all op compose.
 	//Maybe don't want to make that externally usable though? Maybe for power users...
 	TArray<TNiagaraExprPtr> Columns;
-	Columns.Add(Expression_VMNative(VectorVM::EOp::composex, Row0, Row1, Row2, Row3));
-	Columns.Add(Expression_VMNative(VectorVM::EOp::composey, Row0, Row1, Row2, Row3));
-	Columns.Add(Expression_VMNative(VectorVM::EOp::composez, Row0, Row1, Row2, Row3));
-	Columns.Add(Expression_VMNative(VectorVM::EOp::composew, Row0, Row1, Row2, Row3));
+	Columns.Add(Expression_VMNative(EVectorVMOp::composex, Row0, Row1, Row2, Row3));
+	Columns.Add(Expression_VMNative(EVectorVMOp::composey, Row0, Row1, Row2, Row3));
+	Columns.Add(Expression_VMNative(EVectorVMOp::composez, Row0, Row1, Row2, Row3));
+	Columns.Add(Expression_VMNative(EVectorVMOp::composew, Row0, Row1, Row2, Row3));
 
 	Matrix_Internal(Columns, OutputExpressions);
 	return true;
@@ -687,39 +1018,49 @@ bool FNiagaraCompiler_VectorVM::Inverse_Internal(TArray<TNiagaraExprPtr>& InputE
 
 bool FNiagaraCompiler_VectorVM::LessThan_Internal(TArray<TNiagaraExprPtr>& InputExpressions, TArray<TNiagaraExprPtr>& OutputExpressions)
 {
-	OutputExpressions.Add(Expression_VMNative(VectorVM::EOp::lessthan, InputExpressions));
+	OutputExpressions.Add(Expression_VMNative(EVectorVMOp::lessthan, InputExpressions));
 	return true;
 }
 
+bool FNiagaraCompiler_VectorVM::GreaterThan_Internal(TArray<TNiagaraExprPtr>& InputExpressions, TArray<TNiagaraExprPtr>& OutputExpressions)
+{
+	check(InputExpressions.Num() == 2);
+	//Use lessthan op with swapped inputs.
+	TArray<TNiagaraExprPtr> SwappedInputs;
+	SwappedInputs.Add(InputExpressions[1]);
+	SwappedInputs.Add(InputExpressions[0]);
+	OutputExpressions.Add(Expression_VMNative(EVectorVMOp::lessthan, SwappedInputs));
+	return true;
+}
 
-bool FNiagaraCompiler_VectorVM::Sample_Internal(TArray<TNiagaraExprPtr>& InputExpressions, TArray<TNiagaraExprPtr>& OutputExpressions)
+bool FNiagaraCompiler_VectorVM::Select_Internal(TArray<TNiagaraExprPtr>& InputExpressions, TArray<TNiagaraExprPtr>& OutputExpressions)
 {
-	OutputExpressions.Add(Expression_VMNative(VectorVM::EOp::sample, InputExpressions));
+	OutputExpressions.Add(Expression_VMNative(EVectorVMOp::select, InputExpressions));
 	return true;
 }
 
-bool FNiagaraCompiler_VectorVM::Write_Internal(TArray<TNiagaraExprPtr>& InputExpressions, TArray<TNiagaraExprPtr>& OutputExpressions)
+bool FNiagaraCompiler_VectorVM::Sample_Internal(TArray<TNiagaraExprPtr>& InputExpressions, TArray<TNiagaraExprPtr>& OutputExpressions)
 {
-	OutputExpressions.Add(Expression_VMNative(VectorVM::EOp::bufferwrite, InputExpressions));
+	OutputExpressions.Add(Expression_VMNative(EVectorVMOp::sample, InputExpressions));
 	return true;
 }
 
-bool FNiagaraCompiler_VectorVM::EventBroadcast_Internal(TArray<TNiagaraExprPtr>& InputExpressions, TArray<TNiagaraExprPtr>& OutputExpressions)
+bool FNiagaraCompiler_VectorVM::Write_Internal(TArray<TNiagaraExprPtr>& InputExpressions, TArray<TNiagaraExprPtr>& OutputExpressions)
 {
-	OutputExpressions.Add(Expression_VMNative(VectorVM::EOp::eventbroadcast, InputExpressions));
+	OutputExpressions.Add(Expression_VMNative(EVectorVMOp::bufferwrite, InputExpressions));
 	return true;
 }
 
 bool FNiagaraCompiler_VectorVM::EaseIn_Internal(TArray<TNiagaraExprPtr>& InputExpressions, TArray<TNiagaraExprPtr>& OutputExpressions)
 {
-	OutputExpressions.Add(Expression_VMNative(VectorVM::EOp::easein, InputExpressions));
+	OutputExpressions.Add(Expression_VMNative(EVectorVMOp::easein, InputExpressions));
 	return true;
 }
 
 
 bool FNiagaraCompiler_VectorVM::EaseInOut_Internal(TArray<TNiagaraExprPtr>& InputExpressions, TArray<TNiagaraExprPtr>& OutputExpressions)
 {
-	OutputExpressions.Add(Expression_VMNative(VectorVM::EOp::easeinout, InputExpressions));
+	OutputExpressions.Add(Expression_VMNative(EVectorVMOp::easeinout, InputExpressions));
 	return true;
 }
 
diff --git a/Engine/Source/Editor/NiagaraEditor/Private/NiagaraCompiler_VectorVM.h b/Engine/Source/Editor/NiagaraEditor/Private/NiagaraCompiler_VectorVM.h
--- a/Engine/Source/Editor/NiagaraEditor/Private/NiagaraCompiler_VectorVM.h
+++ b/Engine/Source/Editor/NiagaraEditor/Private/NiagaraCompiler_VectorVM.h
@@ -11,11 +11,11 @@ class NIAGARAEDITOR_API FNiagaraCompiler_VectorVM : public FNiagaraCompiler
 private:
 
 	/** Creates an expression directly mapping to a VM operation. */
-	TNiagaraExprPtr Expression_VMNative(VectorVM::EOp Op, TArray<TNiagaraExprPtr>& InputExpressions);
-	TNiagaraExprPtr Expression_VMNative(VectorVM::EOp Op, TNiagaraExprPtr A);
-	TNiagaraExprPtr Expression_VMNative(VectorVM::EOp Op, TNiagaraExprPtr A, TNiagaraExprPtr B);
-	TNiagaraExprPtr Expression_VMNative(VectorVM::EOp Op, TNiagaraExprPtr A, TNiagaraExprPtr B, TNiagaraExprPtr C);
-	TNiagaraExprPtr Expression_VMNative(VectorVM::EOp Op, TNiagaraExprPtr A, TNiagaraExprPtr B, TNiagaraExprPtr C, TNiagaraExprPtr D);
+	TNiagaraExprPtr Expression_VMNative(EVectorVMOp Op, TArray<TNiagaraExprPtr>& InputExpressions);
+	TNiagaraExprPtr Expression_VMNative(EVectorVMOp Op, TNiagaraExprPtr A);
+	TNiagaraExprPtr Expression_VMNative(EVectorVMOp Op, TNiagaraExprPtr A, TNiagaraExprPtr B);
+	TNiagaraExprPtr Expression_VMNative(EVectorVMOp Op, TNiagaraExprPtr A, TNiagaraExprPtr B, TNiagaraExprPtr C);
+	TNiagaraExprPtr Expression_VMNative(EVectorVMOp Op, TNiagaraExprPtr A, TNiagaraExprPtr B, TNiagaraExprPtr C, TNiagaraExprPtr D);
 
 	/** Array tracking which temp registers are free allowing their re-use and thus a smaller total number. */
 	TArray<bool> TempRegisters;
@@ -24,7 +24,7 @@ public:
 	/** Returns true if the expression is stored in the constant data. */
 	bool ExpressionIsConstant(FNiagaraExpression*  Expression);
 	bool ExpressionIsBufferConstant(FNiagaraExpression*  Expression);
-	VectorVM::EOperandType GetVMOperandType(FNiagaraExpression *Expression);
+	EVectorVMOperandLocation GetVMOperandType(FNiagaraExpression *Expression);
 
 	/** Gets the final index used in the VM bytecode for the passed expression. */
 	uint8 GetResultVMIndex(FNiagaraExpression*  Expression);
@@ -40,10 +40,17 @@ public:
 #undef NiagaraOp
 
 	virtual TNiagaraExprPtr Output(const FNiagaraVariableInfo& Attr, TNiagaraExprPtr& SourceExpression)override;
+	virtual TNiagaraExprPtr SharedDataRead(const FNiagaraDataSetID& DataSet, const FNiagaraVariableInfo& Variable, TNiagaraExprPtr& IndexExpr) override;
+	virtual TNiagaraExprPtr SharedDataWrite(const FNiagaraDataSetID& DataSet, const FNiagaraVariableInfo& Variable, TNiagaraExprPtr& IndexExpr, TNiagaraExprPtr& DataExpr) override;
+
+
+	virtual TNiagaraExprPtr AcquireSharedDataIndex(const FNiagaraDataSetID& DataSet, bool bWrap, TNiagaraExprPtr& ValidExpr) override;
+	virtual TNiagaraExprPtr ConsumeSharedDataIndex(const FNiagaraDataSetID& DataSet, bool bWrap) override;
+	virtual TNiagaraExprPtr SharedDataIndexIsValid(const FNiagaraDataSetID& DataSet, TNiagaraExprPtr& IndexExpr) override;
 	//End INiagaraCompiler;
 
 	//Begin FNiagaraCompiler Interface
-	virtual ENiagaraDataType GetConstantResultIndex(const FNiagaraVariableInfo& Constant, bool bInternal, int32& OutResultIndex, int32& OutComponentIndex)override;
+	virtual ENiagaraDataType GetConstantResultIndex(const FNiagaraVariableInfo& Constant, bool bInternal, int32& OutResultIndex)override;
 	virtual int32 AquireTemporary()override;
 	virtual void FreeTemporary(int32 TempIndex) override;
 	//End INiagaraCompiler Interface
diff --git a/Engine/Source/Editor/NiagaraEditor/Private/NiagaraEditor.cpp b/Engine/Source/Editor/NiagaraEditor/Private/NiagaraEditor.cpp
--- a/Engine/Source/Editor/NiagaraEditor/Private/NiagaraEditor.cpp
+++ b/Engine/Source/Editor/NiagaraEditor/Private/NiagaraEditor.cpp
@@ -16,7 +16,8 @@
  
 #define LOCTEXT_NAMESPACE "NiagaraEditor"
 
-const FName FNiagaraEditor::NodeGraphTabId(TEXT("NiagaraEditor_NodeGraph"));
+const FName FNiagaraEditor::NodeGraphTabId(TEXT("NiagaraEditor_NodeGraph")); 
+const FName FNiagaraEditor::FlattenedNodeGraphTabId(TEXT("NiagaraEditor_FlattenedNodeGraph")); 
 const FName FNiagaraEditor::PropertiesTabId(TEXT("NiagaraEditor_MaterialProperties"));
 
 void FNiagaraEditor::RegisterTabSpawners(const TSharedRef<class FTabManager>& TabManager)
@@ -29,7 +30,11 @@ void FNiagaraEditor::RegisterTabSpawners(const TSharedRef<class FTabManager>& Ta
 
 	TabManager->RegisterTabSpawner( NodeGraphTabId, FOnSpawnTab::CreateSP(this, &FNiagaraEditor::SpawnTab_NodeGraph) )
 		.SetDisplayName( LOCTEXT("NodeGraph", "Node Graph") )
-		.SetGroup(WorkspaceMenuCategoryRef);
+		.SetGroup(WorkspaceMenuCategoryRef); 
+
+	TabManager->RegisterTabSpawner( FlattenedNodeGraphTabId, FOnSpawnTab::CreateSP(this, &FNiagaraEditor::SpawnTab_FlattenedNodeGraph))
+		.SetDisplayName(LOCTEXT("FlattenedNodeGraph", "Flattened Node Graph"))
+		.SetGroup(WorkspaceMenuCategoryRef); 
 
 	TabManager->RegisterTabSpawner(PropertiesTabId, FOnSpawnTab::CreateSP(this, &FNiagaraEditor::SpawnTab_NodeProperties))
 		.SetDisplayName(LOCTEXT("DetailsTab", "Details"))
@@ -133,7 +138,7 @@ FLinearColor FNiagaraEditor::GetWorldCentricTabColorScale() const
 
 
 /** Create new tab for the supplied graph - don't call this directly, call SExplorer->FindTabForGraph.*/
-TSharedRef<SGraphEditor> FNiagaraEditor::CreateGraphEditorWidget(UEdGraph* InGraph)
+TSharedRef<SGraphEditor> FNiagaraEditor::CreateGraphEditorWidget(UEdGraph* InGraph, bool bEditable)
 {
 	check(InGraph != NULL);
 	
@@ -203,7 +208,8 @@ TSharedRef<SGraphEditor> FNiagaraEditor::CreateGraphEditorWidget(UEdGraph* InGra
 		.Appearance(AppearanceInfo)
 		.TitleBar(TitleBarWidget)
 		.GraphToEdit(InGraph)
-		.GraphEvents(InEvents);
+		.GraphEvents(InEvents)
+		.IsEditable(bEditable);
 }
 
 
@@ -211,7 +217,7 @@ TSharedRef<SDockTab> FNiagaraEditor::SpawnTab_NodeGraph( const FSpawnTabArgs& Ar
 {
 	check( Args.GetTabId().TabType == NodeGraphTabId );
 
-	TSharedRef<SGraphEditor> NodeGraphEditor = CreateGraphEditorWidget(Source->NodeGraph);
+	TSharedRef<SGraphEditor> NodeGraphEditor = CreateGraphEditorWidget(Source->NodeGraph, true);
 
 	NodeGraphEditorPtr = NodeGraphEditor; // Keep pointer to editor
 
@@ -223,6 +229,22 @@ TSharedRef<SDockTab> FNiagaraEditor::SpawnTab_NodeGraph( const FSpawnTabArgs& Ar
 		];
 }
 
+TSharedRef<SDockTab> FNiagaraEditor::SpawnTab_FlattenedNodeGraph(const FSpawnTabArgs& Args)
+{
+	check(Args.GetTabId().TabType == FlattenedNodeGraphTabId);
+
+	TSharedRef<SGraphEditor> NodeGraphEditor = CreateGraphEditorWidget(Source->FlattenedNodeGraph, false);
+
+	FlattenedNodeGraphEditorPtr = NodeGraphEditor; // Keep pointer to editor
+
+	return SNew(SDockTab)
+		.Label(LOCTEXT("Flattened NodeGraph", "Flattened Node Graph"))
+		.TabColorScale(GetTabColorScale())
+		[
+			NodeGraphEditor
+		];
+}
+
 TSharedRef<SDockTab> FNiagaraEditor::SpawnTab_NodeProperties(const FSpawnTabArgs& Args)
 {
 	TSharedRef<SDockTab> SpawnedTab = SNew(SDockTab)
@@ -299,6 +321,15 @@ void FNiagaraEditor::ExtendToolbar()
 FReply FNiagaraEditor::OnCompileClicked()
 {
 	Script->Source->Compile();
+
+	TSharedPtr<SGraphEditor> CurrentGraphEditor = NodeGraphEditorPtr.Pin();
+	if (CurrentGraphEditor.IsValid())
+		CurrentGraphEditor->NotifyGraphChanged();
+
+	TSharedPtr<SGraphEditor> FlattenedGraphEditor = FlattenedNodeGraphEditorPtr.Pin();
+	if (FlattenedGraphEditor.IsValid())
+		FlattenedGraphEditor->NotifyGraphChanged();
+
 	return FReply::Handled();
 }
 
@@ -398,16 +429,18 @@ bool FNiagaraEditor::CanDeleteNodes() const
 {
 	// If any of the nodes can be deleted then we should allow deleting
 	TSharedPtr<SGraphEditor> CurrentGraphEditor = NodeGraphEditorPtr.Pin();
-	const FGraphPanelSelectionSet SelectedNodes = CurrentGraphEditor->GetSelectedNodes();
-	for (FGraphPanelSelectionSet::TConstIterator SelectedIter(SelectedNodes); SelectedIter; ++SelectedIter)
+	if (CurrentGraphEditor.IsValid())
 	{
-		UEdGraphNode* Node = Cast<UEdGraphNode>(*SelectedIter);
-		if (Node && Node->CanUserDeleteNode())
+		const FGraphPanelSelectionSet SelectedNodes = CurrentGraphEditor->GetSelectedNodes();
+		for (FGraphPanelSelectionSet::TConstIterator SelectedIter(SelectedNodes); SelectedIter; ++SelectedIter)
 		{
-			return true;
+			UEdGraphNode* Node = Cast<UEdGraphNode>(*SelectedIter);
+			if (Node && Node->CanUserDeleteNode())
+			{
+				return true;
+			}
 		}
 	}
-
 	return false;
 }
 
diff --git a/Engine/Source/Editor/NiagaraEditor/Private/NiagaraEditor.h b/Engine/Source/Editor/NiagaraEditor/Private/NiagaraEditor.h
--- a/Engine/Source/Editor/NiagaraEditor/Private/NiagaraEditor.h
+++ b/Engine/Source/Editor/NiagaraEditor/Private/NiagaraEditor.h
@@ -42,11 +42,14 @@ protected:
 
 private:
 	/** Create widget for graph editing */
-	TSharedRef<class SGraphEditor> CreateGraphEditorWidget(UEdGraph* InGraph);
+	TSharedRef<class SGraphEditor> CreateGraphEditorWidget(UEdGraph* InGraph, bool bEditable = true);
 
 	/** Spawns the tab with the update graph inside */
 	TSharedRef<SDockTab> SpawnTab_NodeGraph(const FSpawnTabArgs& Args);
 
+	/** Spawns the tab with the flattened graph inside */
+	TSharedRef<SDockTab> SpawnTab_FlattenedNodeGraph(const FSpawnTabArgs& Args);
+
 	/** Spawns the tab with the node details inside. */
 	TSharedRef<SDockTab> SpawnTab_NodeProperties(const FSpawnTabArgs& Args);
 
@@ -79,6 +82,7 @@ private:
 	UNiagaraScriptSource*		Source;
 	/** */
 	TWeakPtr<SGraphEditor>		NodeGraphEditorPtr;
+	TWeakPtr<SGraphEditor>		FlattenedNodeGraphEditorPtr;
 
 	/** The command list for this editor */
 	TSharedPtr<FUICommandList> GraphEditorCommands;
@@ -87,7 +91,8 @@ private:
 	TSharedPtr<class IDetailsView> NiagaraDetailsView;
 
 	/**	The tab ids for the Niagara editor */
-	static const FName NodeGraphTabId;
+	static const FName NodeGraphTabId; 
+	static const FName FlattenedNodeGraphTabId; 
 	static const FName PropertiesTabId;
 
 };
diff --git a/Engine/Source/Editor/NiagaraEditor/Private/NiagaraEffectEditor.cpp b/Engine/Source/Editor/NiagaraEditor/Private/NiagaraEffectEditor.cpp
--- a/Engine/Source/Editor/NiagaraEditor/Private/NiagaraEffectEditor.cpp
+++ b/Engine/Source/Editor/NiagaraEditor/Private/NiagaraEffectEditor.cpp
@@ -78,7 +78,7 @@ void FNiagaraEffectEditor::InitNiagaraEffectEditor(const EToolkitMode::Type Mode
 {
 	Effect = InEffect;
 	check(Effect != NULL);
-	EffectInstance = new FNiagaraEffectInstance(InEffect);
+	EffectInstance = MakeShareable(new FNiagaraEffectInstance(InEffect));
 
 	const float InTime = -0.02f;
 	const float OutTime = 3.2f;
@@ -187,20 +187,21 @@ FLinearColor FNiagaraEffectEditor::GetWorldCentricTabColorScale() const
 	return FLinearColor(0.0f, 0.0f, 0.2f, 0.5f);
 }
 
+void FNiagaraEffectEditor::NotifyPreChange(UProperty* PropertyAboutToChanged)
+{
+	Viewport->GetPreviewComponent()->UnregisterComponent();
+}
+
 void FNiagaraEffectEditor::NotifyPostChange(const FPropertyChangedEvent& PropertyChangedEvent, UProperty* PropertyThatChanged)
 {
-	//Rebuild the effect after every edit.
-// 	if (EffectInstance)
-// 	{
-// 		EffectInstance->Init();
-// 	}
+	Viewport->GetPreviewComponent()->RegisterComponent();
 }
 
 /** Create new tab for the supplied graph - don't call this directly, call SExplorer->FindTabForGraph.*/
 TSharedRef<SNiagaraEffectEditorWidget> FNiagaraEffectEditor::CreateEditorWidget(UNiagaraEffect* InEffect)
 {
 	check(InEffect != NULL);
-	EffectInstance = new FNiagaraEffectInstance(InEffect);
+	EffectInstance = MakeShareable(new FNiagaraEffectInstance(InEffect));
 	
 	if (!EditorCommands.IsValid())
 	{
@@ -231,7 +232,7 @@ TSharedRef<SNiagaraEffectEditorWidget> FNiagaraEffectEditor::CreateEditorWidget(
 		
 		
 	// Make the effect editor widget
-	return SNew(SNiagaraEffectEditorWidget).EffectObj(InEffect).EffectInstance(EffectInstance).EffectEditor(this);
+	return SNew(SNiagaraEffectEditorWidget).EffectObj(InEffect).EffectInstance(EffectInstance.Get()).EffectObj(Effect).EffectEditor(this);
 }
 
 
@@ -281,7 +282,7 @@ TSharedRef<SDockTab> FNiagaraEffectEditor::SpawnTab_EmitterList(const FSpawnTabA
 	TSharedRef<SDockTab> SpawnedTab =
 		SNew(SDockTab)
 		[
-			SAssignNew(EmitterEditorWidget, SNiagaraEffectEditorWidget).EffectInstance(EffectInstance).EffectEditor(this)
+			SAssignNew(EmitterEditorWidget, SNiagaraEffectEditorWidget).EffectInstance(EffectInstance.Get()).EffectEditor(this).EffectObj(Effect)
 		];
 
 
@@ -295,7 +296,7 @@ TSharedRef<SDockTab> FNiagaraEffectEditor::SpawnTab_DevEmitterList(const FSpawnT
 	TSharedRef<SDockTab> SpawnedTab =
 		SNew(SDockTab)
 		[
-			SAssignNew(DevEmitterEditorWidget, SNiagaraEffectEditorWidget).EffectInstance(EffectInstance).EffectEditor(this).bForDev(true)
+			SAssignNew(DevEmitterEditorWidget, SNiagaraEffectEditorWidget).EffectInstance(EffectInstance.Get()).EffectEditor(this).EffectObj(Effect).bForDev(true)
 		];
 
 
diff --git a/Engine/Source/Editor/NiagaraEditor/Private/NiagaraEffectEditor.h b/Engine/Source/Editor/NiagaraEditor/Private/NiagaraEffectEditor.h
--- a/Engine/Source/Editor/NiagaraEditor/Private/NiagaraEffectEditor.h
+++ b/Engine/Source/Editor/NiagaraEditor/Private/NiagaraEffectEditor.h
@@ -44,6 +44,7 @@ public:
 		return MakeShareable(new FNiagaraTrackEditor(InSequencer) );
 	}
 
+	virtual void NotifyPreChange(UProperty* PropertyAboutToChanged)override;
 	virtual void NotifyPostChange(const FPropertyChangedEvent& PropertyChangedEvent, UProperty* PropertyThatChanged) override;
 
 	FReply OnDeleteEmitterClicked(TSharedPtr<FNiagaraSimulation> Emitter);
@@ -74,7 +75,7 @@ private:
 
 	/* The Effect being edited */
 	UNiagaraEffect	*Effect;
-	class FNiagaraEffectInstance *EffectInstance;
+	TSharedPtr<FNiagaraEffectInstance> EffectInstance;
 
 	/* stuff needed by the Sequencer */
 	UMovieScene *MovieScene;
