** PRIMARY FEATURE **

diff --git a/Engine/Source/Editor/NiagaraEditor/Private/NiagaraEmitterPropertiesDetailsCustomization.cpp b/Engine/Source/Editor/NiagaraEditor/Private/NiagaraEmitterPropertiesDetailsCustomization.cpp
--- a/Engine/Source/Editor/NiagaraEditor/Private/NiagaraEmitterPropertiesDetailsCustomization.cpp
+++ b/Engine/Source/Editor/NiagaraEditor/Private/NiagaraEmitterPropertiesDetailsCustomization.cpp
@@ -24,81 +24,74 @@ FNiagaraEmitterPropertiesDetails::FNiagaraEmitterPropertiesDetails(TWeakObjectPt
 {
 }
 
-void FNiagaraEmitterPropertiesDetails::CustomizeDetails(IDetailLayoutBuilder& DetailLayout)
+void FNiagaraEmitterPropertiesDetails::CustomizeDetails(IDetailLayoutBuilder& InDetailLayout)
 {
+	DetailLayout = &InDetailLayout;
 	FName EmitterCategoryName = TEXT("Emitter");
-	IDetailCategoryBuilder& EmitterCategory = DetailLayout.EditCategory(EmitterCategoryName);
+	IDetailCategoryBuilder& EmitterCategory = DetailLayout->EditCategory(EmitterCategoryName);
 
-	TSharedRef<IPropertyHandle> EmitterName = DetailLayout.GetProperty(GET_MEMBER_NAME_CHECKED(UNiagaraEmitterProperties, EmitterName));
-	TSharedRef<IPropertyHandle> EmitterEnabled = DetailLayout.GetProperty(GET_MEMBER_NAME_CHECKED(UNiagaraEmitterProperties, bIsEnabled));
+	TSharedRef<IPropertyHandle> EmitterName = DetailLayout->GetProperty(GET_MEMBER_NAME_CHECKED(UNiagaraEmitterProperties, EmitterName));
+	TSharedRef<IPropertyHandle> EmitterEnabled = DetailLayout->GetProperty(GET_MEMBER_NAME_CHECKED(UNiagaraEmitterProperties, bIsEnabled));
 
-	DetailLayout.HideProperty(EmitterName);
-	DetailLayout.HideProperty(EmitterEnabled);
+	DetailLayout->HideProperty(EmitterName);
+	DetailLayout->HideProperty(EmitterEnabled);
 	
-	TSharedRef<IPropertyHandle> UpdateScriptProps = DetailLayout.GetProperty(GET_MEMBER_NAME_CHECKED(UNiagaraEmitterProperties, UpdateScriptProps));
-	TSharedRef<IPropertyHandle> SpawnScriptProps = DetailLayout.GetProperty(GET_MEMBER_NAME_CHECKED(UNiagaraEmitterProperties, SpawnScriptProps));
+	TSharedRef<IPropertyHandle> UpdateScriptProps = DetailLayout->GetProperty(GET_MEMBER_NAME_CHECKED(UNiagaraEmitterProperties, UpdateScriptProps));
+	TSharedRef<IPropertyHandle> SpawnScriptProps = DetailLayout->GetProperty(GET_MEMBER_NAME_CHECKED(UNiagaraEmitterProperties, SpawnScriptProps));
 
-	BuildScriptProperties(DetailLayout, SpawnScriptProps, TEXT("Spawn Script"), LOCTEXT("SpawnScriptDisplayName", "Spawn Script"));
-	BuildScriptProperties(DetailLayout, UpdateScriptProps, TEXT("Update Script"), LOCTEXT("UpdateScriptDisplayName", "Update Script"));
+	BuildScriptProperties(SpawnScriptProps, TEXT("Spawn Script"), LOCTEXT("SpawnScriptDisplayName", "Spawn Script"));
+	BuildScriptProperties(UpdateScriptProps, TEXT("Update Script"), LOCTEXT("UpdateScriptDisplayName", "Update Script"));
 }
 
-void FNiagaraEmitterPropertiesDetails::BuildScriptProperties(IDetailLayoutBuilder& DetailLayout, TSharedRef<IPropertyHandle> ScriptPropsHandle, FName Name, FText DisplayName)
+void FNiagaraEmitterPropertiesDetails::BuildScriptProperties(TSharedRef<IPropertyHandle> ScriptPropsHandle, FName Name, FText DisplayName)
 {
-	TFunction<void(TSharedRef<IPropertyHandle>)> HideAllBelow = [&](TSharedRef<IPropertyHandle> Prop)
-	{
-		DetailLayout.HideProperty(Prop);
-
-		//Don't hide sub properties for data objects.
-		if (Prop->GetProperty()->GetName() == GET_MEMBER_NAME_CHECKED(FNiagaraConstants, DataObjectConstants).ToString())
-			return; 
-
-		uint32 NumChildren = 0;
-		if (Prop->GetNumChildren(NumChildren) == FPropertyAccess::Success)
-		{
-			for (uint32 i = 0; i < NumChildren; ++i)
-			{
-				TSharedPtr<IPropertyHandle> Child = Prop->GetChildHandle(i);
-				if (Child.IsValid())
-				{
-					HideAllBelow(Child.ToSharedRef());
-				}				
-			}
-		}
-	};
-	HideAllBelow(ScriptPropsHandle);
+	DetailLayout->HideProperty(ScriptPropsHandle);
 
-	IDetailCategoryBuilder& ScriptCategory = DetailLayout.EditCategory(Name, DisplayName);
+	IDetailCategoryBuilder& ScriptCategory = DetailLayout->EditCategory(Name, DisplayName);
 
 	//Script
 	TSharedPtr<IPropertyHandle> ScriptHandle = ScriptPropsHandle->GetChildHandle(GET_MEMBER_NAME_CHECKED(FNiagaraEmitterScriptProperties, Script));
 	ScriptCategory.AddProperty(ScriptHandle);
 
 	//Constants
-	const bool bGenerateHeader = false;
-	const bool bDisplayResetToDefault = false;
+	bool bGenerateHeader = false;
+	bool bDisplayResetToDefault = false;
 	TSharedPtr<IPropertyHandle> Constants = ScriptPropsHandle->GetChildHandle(GET_MEMBER_NAME_CHECKED(FNiagaraEmitterScriptProperties, ExternalConstants));
+	DetailLayout->HideProperty(Constants);
 	//Scalar Constants.
-	TSharedPtr<IPropertyHandle> ScalarConstants = ScriptPropsHandle->GetChildHandle(GET_MEMBER_NAME_CHECKED(FNiagaraConstants, ScalarConstants));
+	TSharedPtr<IPropertyHandle> ScalarConstants = Constants->GetChildHandle(GET_MEMBER_NAME_CHECKED(FNiagaraConstants, ScalarConstants));
 	TSharedRef<FDetailArrayBuilder> ScalarConstantsBuilder = MakeShareable(new FDetailArrayBuilder(ScalarConstants.ToSharedRef(), bGenerateHeader, bDisplayResetToDefault));
 	ScalarConstantsBuilder->OnGenerateArrayElementWidget(FOnGenerateArrayElementWidget::CreateSP(this, &FNiagaraEmitterPropertiesDetails::OnGenerateScalarConstantEntry));
 	ScriptCategory.AddCustomBuilder(ScalarConstantsBuilder);
 	//Vector Constants.
-	TSharedPtr<IPropertyHandle> VectorConstants = ScriptPropsHandle->GetChildHandle(GET_MEMBER_NAME_CHECKED(FNiagaraConstants, VectorConstants));
+	TSharedPtr<IPropertyHandle> VectorConstants = Constants->GetChildHandle(GET_MEMBER_NAME_CHECKED(FNiagaraConstants, VectorConstants));
 	TSharedRef<FDetailArrayBuilder> VectorConstantsBuilder = MakeShareable(new FDetailArrayBuilder(VectorConstants.ToSharedRef(), bGenerateHeader, bDisplayResetToDefault));
 	VectorConstantsBuilder->OnGenerateArrayElementWidget(FOnGenerateArrayElementWidget::CreateSP(this, &FNiagaraEmitterPropertiesDetails::OnGenerateVectorConstantEntry));
 	ScriptCategory.AddCustomBuilder(VectorConstantsBuilder);
 	//Matrix Constants.
-	TSharedPtr<IPropertyHandle> MatrixConstants = ScriptPropsHandle->GetChildHandle(GET_MEMBER_NAME_CHECKED(FNiagaraConstants, MatrixConstants));
+	TSharedPtr<IPropertyHandle> MatrixConstants = Constants->GetChildHandle(GET_MEMBER_NAME_CHECKED(FNiagaraConstants, MatrixConstants));
 	TSharedRef<FDetailArrayBuilder> MatrixConstantsBuilder = MakeShareable(new FDetailArrayBuilder(MatrixConstants.ToSharedRef(), bGenerateHeader, bDisplayResetToDefault));
 	MatrixConstantsBuilder->OnGenerateArrayElementWidget(FOnGenerateArrayElementWidget::CreateSP(this, &FNiagaraEmitterPropertiesDetails::OnGenerateMatrixConstantEntry));
 	ScriptCategory.AddCustomBuilder(MatrixConstantsBuilder);
 	//DataObject Constants.
-	TSharedPtr<IPropertyHandle> DataObjectConstants = ScriptPropsHandle->GetChildHandle(GET_MEMBER_NAME_CHECKED(FNiagaraConstants, DataObjectConstants));
+	TSharedPtr<IPropertyHandle> DataObjectConstants = Constants->GetChildHandle(GET_MEMBER_NAME_CHECKED(FNiagaraConstants, DataObjectConstants));
 	TSharedRef<FDetailArrayBuilder> DataObjectConstantsBuilder = MakeShareable(new FDetailArrayBuilder(DataObjectConstants.ToSharedRef(), bGenerateHeader, bDisplayResetToDefault));
 	DataObjectConstantsBuilder->OnGenerateArrayElementWidget(FOnGenerateArrayElementWidget::CreateSP(this, &FNiagaraEmitterPropertiesDetails::OnGenerateDataObjectConstantEntry));
 	ScriptCategory.AddCustomBuilder(DataObjectConstantsBuilder);
 
-	//TODO: Events.
+	//Events
+	bGenerateHeader = true;
+	//Generators
+	TSharedPtr<IPropertyHandle> EventGenerators = ScriptPropsHandle->GetChildHandle(GET_MEMBER_NAME_CHECKED(FNiagaraEmitterScriptProperties, EventGenerators));
+	TSharedRef<FDetailArrayBuilder> EventGeneratorsBuilder = MakeShareable(new FDetailArrayBuilder(EventGenerators.ToSharedRef(), bGenerateHeader, bDisplayResetToDefault));
+	EventGeneratorsBuilder->OnGenerateArrayElementWidget(FOnGenerateArrayElementWidget::CreateSP(this, &FNiagaraEmitterPropertiesDetails::OnGenerateEventGeneratorEntry));
+	ScriptCategory.AddCustomBuilder(EventGeneratorsBuilder);
+	//Receivers
+	TSharedPtr<IPropertyHandle> EventReceivers = ScriptPropsHandle->GetChildHandle(GET_MEMBER_NAME_CHECKED(FNiagaraEmitterScriptProperties, EventReceivers));
+	TSharedRef<FDetailArrayBuilder> EventReceiversBuilder = MakeShareable(new FDetailArrayBuilder(EventReceivers.ToSharedRef(), bGenerateHeader, bDisplayResetToDefault));
+	EventReceiversBuilder->OnGenerateArrayElementWidget(FOnGenerateArrayElementWidget::CreateSP(this, &FNiagaraEmitterPropertiesDetails::OnGenerateEventReceiverEntry));
+	ScriptCategory.AddCustomBuilder(EventReceiversBuilder);
+
 }
 
 void FNiagaraEmitterPropertiesDetails::OnGenerateScalarConstantEntry(TSharedRef<IPropertyHandle> ElementProperty, int32 ElementIndex, IDetailChildrenBuilder& ChildrenBuilder)
@@ -157,5 +150,68 @@ void FNiagaraEmitterPropertiesDetails::OnGenerateDataObjectConstantEntry(TShared
 	}
 }
 
+void FNiagaraEmitterPropertiesDetails::OnGenerateEventGeneratorEntry(TSharedRef<IPropertyHandle> ElementProperty, int32 ElementIndex, IDetailChildrenBuilder& ChildrenBuilder)
+{
+	TSharedPtr<IPropertyHandle> IDProperty = ElementProperty->GetChildHandle(GET_MEMBER_NAME_CHECKED(FNiagaraEventGeneratorProperties, ID));
+	TSharedPtr<IPropertyHandle> NameProperty = IDProperty->GetChildHandle(GET_MEMBER_NAME_CHECKED(FNiagaraDataSetID, Name));
+
+	FName DisplayName;
+	NameProperty->GetValue(DisplayName);
+	//IDetailPropertyRow& Row = ChildrenBuilder.AddChildProperty(ElementProperty).DisplayName(FText::FromName(DisplayName));
+
+	IDetailGroup& GenGroup = ChildrenBuilder.AddChildGroup(DisplayName, FText::FromName(DisplayName));
+	uint32 NumChildren = 0;
+	if (ElementProperty->GetNumChildren(NumChildren) == FPropertyAccess::Success)
+	{
+		for (uint32 i = 0; i < NumChildren; ++i)
+		{
+			TSharedPtr<IPropertyHandle> Child = ElementProperty->GetChildHandle(i);
+			//Dont add the ID. We just grab it's name for the name region of this property.
+			if (Child.IsValid() && Child->GetProperty()->GetName() != GET_MEMBER_NAME_CHECKED(FNiagaraEventGeneratorProperties, ID).ToString())
+			{
+				GenGroup.AddPropertyRow(Child.ToSharedRef());
+			}
+		}
+	}
+}
+
+void FNiagaraEmitterPropertiesDetails::OnGenerateEventReceiverEntry(TSharedRef<IPropertyHandle> ElementProperty, int32 ElementIndex, IDetailChildrenBuilder& ChildrenBuilder)
+{
+	TSharedPtr<IPropertyHandle> NameProperty = ElementProperty->GetChildHandle(GET_MEMBER_NAME_CHECKED(FNiagaraEventReceiverProperties, Name));
+
+ 	FName DisplayName;
+ 	NameProperty->GetValue(DisplayName);
+// 	ChildrenBuilder.AddChildProperty(ElementProperty).DisplayName(FText::FromName(DisplayName));
+
+	IDetailGroup& Group = ChildrenBuilder.AddChildGroup(DisplayName, FText::FromName(DisplayName));
+	uint32 NumChildren = 0;
+	if (ElementProperty->GetNumChildren(NumChildren) == FPropertyAccess::Success)
+	{
+		for (uint32 i = 0; i < NumChildren; ++i)
+		{
+			TSharedPtr<IPropertyHandle> Child = ElementProperty->GetChildHandle(i);
+			//Dont add the ID. We just grab it's name for the name region of this property.
+			if (Child.IsValid() && Child->GetProperty()->GetName() != GET_MEMBER_NAME_CHECKED(FNiagaraEventReceiverProperties, Name).ToString())
+			{
+				TSharedPtr<SWidget> NameWidget;
+				TSharedPtr<SWidget> ValueWidget;
+				FDetailWidgetRow DefaultDetailRow;
+			
+				IDetailPropertyRow& Row = Group.AddPropertyRow(Child.ToSharedRef());
+				Row.GetDefaultWidgets(NameWidget, ValueWidget, DefaultDetailRow);
+				Row.CustomWidget(true)
+					.NameContent()
+					[
+						NameWidget.ToSharedRef()
+					]
+					.ValueContent()
+					[
+						ValueWidget.ToSharedRef()
+					];
+			}
+		}
+	}
+}
+
 #undef LOCTEXT_NAMESPACE
 
diff --git a/Engine/Source/Editor/NiagaraEditor/Private/NiagaraEmitterPropertiesDetailsCustomization.h b/Engine/Source/Editor/NiagaraEditor/Private/NiagaraEmitterPropertiesDetailsCustomization.h
--- a/Engine/Source/Editor/NiagaraEditor/Private/NiagaraEmitterPropertiesDetailsCustomization.h
+++ b/Engine/Source/Editor/NiagaraEditor/Private/NiagaraEmitterPropertiesDetailsCustomization.h
@@ -16,16 +16,20 @@ public:
 	FNiagaraEmitterPropertiesDetails(TWeakObjectPtr<UNiagaraEmitterProperties> EmitterProperties);
 
 	/** IDetailCustomization interface */
-	virtual void CustomizeDetails(IDetailLayoutBuilder& DetailLayout) override;
+	virtual void CustomizeDetails(IDetailLayoutBuilder& InDetailLayout) override;
 
 	void OnGenerateScalarConstantEntry(TSharedRef<IPropertyHandle> ElementProperty, int32 ElementIndex, IDetailChildrenBuilder& ChildrenBuilder);
 	void OnGenerateVectorConstantEntry(TSharedRef<IPropertyHandle> ElementProperty, int32 ElementIndex, IDetailChildrenBuilder& ChildrenBuilder);
 	void OnGenerateMatrixConstantEntry(TSharedRef<IPropertyHandle> ElementProperty, int32 ElementIndex, IDetailChildrenBuilder& ChildrenBuilder);
 	void OnGenerateDataObjectConstantEntry(TSharedRef<IPropertyHandle> ElementProperty, int32 ElementIndex, IDetailChildrenBuilder& ChildrenBuilder);
+	void OnGenerateEventGeneratorEntry(TSharedRef<IPropertyHandle> ElementProperty, int32 ElementIndex, IDetailChildrenBuilder& ChildrenBuilder);
+	void OnGenerateEventReceiverEntry(TSharedRef<IPropertyHandle> ElementProperty, int32 ElementIndex, IDetailChildrenBuilder& ChildrenBuilder);
 
-	void BuildScriptProperties(IDetailLayoutBuilder& DetailLayout, TSharedRef<IPropertyHandle> ScriptPropsHandle, FName Name, FText DisplayName);
+	void BuildScriptProperties(TSharedRef<IPropertyHandle> ScriptPropsHandle, FName Name, FText DisplayName);
 private:
 	/** Object that stores all of the possible parameters we can edit */
 	TWeakObjectPtr<UNiagaraEmitterProperties> EmitterProps;
+
+	IDetailLayoutBuilder* DetailLayout;
 };
 
diff --git a/Engine/Source/Editor/NiagaraEditor/Private/SNiagaraEffectEditorViewport.cpp b/Engine/Source/Editor/NiagaraEditor/Private/SNiagaraEffectEditorViewport.cpp
--- a/Engine/Source/Editor/NiagaraEditor/Private/SNiagaraEffectEditorViewport.cpp
+++ b/Engine/Source/Editor/NiagaraEditor/Private/SNiagaraEffectEditorViewport.cpp
@@ -140,7 +140,7 @@ void SNiagaraEffectEditorViewport::Tick( const FGeometry& AllottedGeometry, cons
 	SEditorViewport::Tick( AllottedGeometry, InCurrentTime, InDeltaTime );
 }
 
-void SNiagaraEffectEditorViewport::SetPreviewEffect(FNiagaraEffectInstance *InPreviewEffect)
+void SNiagaraEffectEditorViewport::SetPreviewEffect(TSharedPtr<FNiagaraEffectInstance> InPreviewEffect)
 {
 	check( PreviewComponent );
 
diff --git a/Engine/Source/Editor/NiagaraEditor/Private/SNiagaraEffectEditorViewport.h b/Engine/Source/Editor/NiagaraEditor/Private/SNiagaraEffectEditorViewport.h
--- a/Engine/Source/Editor/NiagaraEditor/Private/SNiagaraEffectEditorViewport.h
+++ b/Engine/Source/Editor/NiagaraEditor/Private/SNiagaraEffectEditorViewport.h
@@ -25,7 +25,7 @@ public:
 
 	void RefreshViewport();
 	
-	void SetPreviewEffect(FNiagaraEffectInstance* InEffect);
+	void SetPreviewEffect(TSharedPtr<FNiagaraEffectInstance> InEffect);
 	
 	void ToggleRealtime();
 	
diff --git a/Engine/Source/Editor/NiagaraEditor/Private/SNiagaraEffectEditorWidget.cpp b/Engine/Source/Editor/NiagaraEditor/Private/SNiagaraEffectEditorWidget.cpp
--- a/Engine/Source/Editor/NiagaraEditor/Private/SNiagaraEffectEditorWidget.cpp
+++ b/Engine/Source/Editor/NiagaraEditor/Private/SNiagaraEffectEditorWidget.cpp
@@ -17,17 +17,13 @@ BEGIN_SLATE_FUNCTION_BUILD_OPTIMIZATION
 void SNiagaraEffectEditorWidget::Construct(const FArguments& InArgs)
 {
 	EffectObj = InArgs._EffectObj;
-	if (InArgs._EffectInstance)
-	{
-		EffectInstance = InArgs._EffectInstance;
-	}
-	else
-	{
-		EffectInstance = new FNiagaraEffectInstance(EffectObj);
-	}
-
+	EffectInstance = InArgs._EffectInstance;
 	EffectEditor = InArgs._EffectEditor;
 	bForDev = InArgs._bForDev;
+
+	check(EffectObj);
+	check(EffectInstance);
+	check(EffectEditor);
 	
 	if (bForDev)
 	{
@@ -136,97 +132,99 @@ void SEmitterWidgetDev::BuildContents()
 	//Register details customization.	
 	if (UNiagaraEmitterProperties* PinnedProps = Emitter->GetProperties().Get())
 	{
-		FPropertyEditorModule& PropertyEditorModule = FModuleManager::GetModuleChecked<FPropertyEditorModule>("PropertyEditor");
-		FDetailsViewArgs DetailsViewArgs(false, false, true, FDetailsViewArgs::HideNameArea, true, this);
-		Details = PropertyEditorModule.CreateDetailView(DetailsViewArgs);
-		FOnGetDetailCustomizationInstance LayoutEmitterDetails = FOnGetDetailCustomizationInstance::CreateStatic(&FNiagaraEmitterPropertiesDetails::MakeInstance, Emitter->GetProperties());
-		Details->RegisterInstancedCustomPropertyLayout(UNiagaraEmitterProperties::StaticClass(), LayoutEmitterDetails);
-
-		Details->SetObject(PinnedProps);
+		if (!Details.IsValid())
+		{
+			FPropertyEditorModule& PropertyEditorModule = FModuleManager::GetModuleChecked<FPropertyEditorModule>("PropertyEditor");
+			FDetailsViewArgs DetailsViewArgs(false, false, true, FDetailsViewArgs::HideNameArea, true, this);
+			Details = PropertyEditorModule.CreateDetailView(DetailsViewArgs);
+			FOnGetDetailCustomizationInstance LayoutEmitterDetails = FOnGetDetailCustomizationInstance::CreateStatic(&FNiagaraEmitterPropertiesDetails::MakeInstance, Emitter->GetProperties());
+			Details->RegisterInstancedCustomPropertyLayout(UNiagaraEmitterProperties::StaticClass(), LayoutEmitterDetails);
 
-		ChildSlot
-			.Padding(4)
-			[
-				SNew(SVerticalBox)
-				+ SVerticalBox::Slot()
-				.AutoHeight()
-				.Padding(0)
+			ChildSlot
+				.Padding(4)
 				[
-					NGED_SECTION_BORDER
+					SNew(SVerticalBox)
+					+ SVerticalBox::Slot()
+					.AutoHeight()
+					.Padding(0)
 					[
-						SNew(SVerticalBox)
-						+ SVerticalBox::Slot()
-						.AutoHeight()
-						.HAlign(HAlign_Fill)
+						NGED_SECTION_BORDER
 						[
-							// name and status line
-							SNew(SHorizontalBox)
-							+ SHorizontalBox::Slot()
-							.HAlign(HAlign_Left)
-							.AutoWidth()
-							.Padding(2)
+							SNew(SVerticalBox)
+							+ SVerticalBox::Slot()
+							.AutoHeight()
+							.HAlign(HAlign_Fill)
 							[
-								SNew(SCheckBox)
-								.OnCheckStateChanged(this, &SEmitterWidgetBase::OnEmitterEnabledChanged)
-								.IsChecked(this, &SEmitterWidgetBase::IsEmitterEnabled)
-								.ToolTipText(FText::FromString("Toggles whether this emitter is enabled. Disabled emitters don't simulate or render."))
-							]
-							+ SHorizontalBox::Slot()
+								// name and status line
+								SNew(SHorizontalBox)
+								+ SHorizontalBox::Slot()
+								.HAlign(HAlign_Left)
 								.AutoWidth()
 								.Padding(2)
-								.HAlign(HAlign_Left)
 								[
-									SNew(SEditableText).Text(this, &SEmitterWidgetBase::GetEmitterName)
-									.MinDesiredWidth(200)
-									.Font(FEditorStyle::GetFontStyle("ContentBrowser.AssetListViewNameFontDirty"))
-									.OnTextChanged(this, &SEmitterWidgetBase::OnEmitterNameChanged)
+									SNew(SCheckBox)
+									.OnCheckStateChanged(this, &SEmitterWidgetBase::OnEmitterEnabledChanged)
+									.IsChecked(this, &SEmitterWidgetBase::IsEmitterEnabled)
+									.ToolTipText(FText::FromString("Toggles whether this emitter is enabled. Disabled emitters don't simulate or render."))
 								]
-							+ SHorizontalBox::Slot()
-								.HAlign(HAlign_Right)
-								.AutoWidth()
-								[
-									SNew(SButton)
-									.VAlign(VAlign_Center)
+								+ SHorizontalBox::Slot()
+									.AutoWidth()
+									.Padding(2)
+									.HAlign(HAlign_Left)
+									[
+										SNew(SEditableText).Text(this, &SEmitterWidgetBase::GetEmitterName)
+										.MinDesiredWidth(200)
+										.Font(FEditorStyle::GetFontStyle("ContentBrowser.AssetListViewNameFontDirty"))
+										.OnTextChanged(this, &SEmitterWidgetBase::OnEmitterNameChanged)
+									]
+								+ SHorizontalBox::Slot()
 									.HAlign(HAlign_Right)
-									.Text(LOCTEXT("DeleteEmitter", "Delete"))
-									.ToolTipText(LOCTEXT("DeleteEmitter_Tooltip", "Deletes this emitter from the effect."))
-									.OnClicked(EffectEditor, &FNiagaraEffectEditor::OnDeleteEmitterClicked, Emitter)
-								]
-							+ SHorizontalBox::Slot()
-								.HAlign(HAlign_Right)
-								.AutoWidth()
-								[
-									SNew(SButton)
-									.VAlign(VAlign_Center)
+									.AutoWidth()
+									[
+										SNew(SButton)
+										.VAlign(VAlign_Center)
+										.HAlign(HAlign_Right)
+										.Text(LOCTEXT("DeleteEmitter", "Delete"))
+										.ToolTipText(LOCTEXT("DeleteEmitter_Tooltip", "Deletes this emitter from the effect."))
+										.OnClicked(EffectEditor, &FNiagaraEffectEditor::OnDeleteEmitterClicked, Emitter)
+									]
+								+ SHorizontalBox::Slot()
 									.HAlign(HAlign_Right)
-									.Text(LOCTEXT("DuplicateEmitter", "Duplicate"))
-									.ToolTipText(LOCTEXT("DuplicateEmitter_Tooltip", "Duplicate this emitter."))
-									.OnClicked(EffectEditor, &FNiagaraEffectEditor::OnDuplicateEmitterClicked, Emitter)
+									.AutoWidth()
+									[
+										SNew(SButton)
+										.VAlign(VAlign_Center)
+										.HAlign(HAlign_Right)
+										.Text(LOCTEXT("DuplicateEmitter", "Duplicate"))
+										.ToolTipText(LOCTEXT("DuplicateEmitter_Tooltip", "Duplicate this emitter."))
+										.OnClicked(EffectEditor, &FNiagaraEffectEditor::OnDuplicateEmitterClicked, Emitter)
+									]
+							]
+							+ SVerticalBox::Slot()
+								.AutoHeight()
+								.HAlign(HAlign_Fill)
+								.Padding(2)
+								[
+									SNew(STextBlock).Text(this, &SEmitterWidgetBase::GetStatsText)
 								]
 						]
-						+ SVerticalBox::Slot()
-							.AutoHeight()
-							.HAlign(HAlign_Fill)
-							.Padding(2)
-							[
-								SNew(STextBlock).Text(this, &SEmitterWidgetBase::GetStatsText)
-							]
 					]
-				]
-				+ SVerticalBox::Slot()
-					.VAlign(VAlign_Fill)
-					.Padding(0)
-					[
-						SNew(SHorizontalBox)
-						+ SHorizontalBox::Slot()						
-						.MaxWidth(400)
+					+ SVerticalBox::Slot()
+						.VAlign(VAlign_Fill)
+						.Padding(0)
 						[
-							Details.ToSharedRef()
+							SNew(SHorizontalBox)
+							+ SHorizontalBox::Slot()
+							.MaxWidth(400)
+							[
+								Details.ToSharedRef()
+							]
 						]
-					]
-			];
-	}
+				];
+		}
 
+		Details->SetObject(PinnedProps, true);
+	}
 }
 
 void SEmitterWidgetDev::NotifyPostChange(const FPropertyChangedEvent& PropertyChangedEvent, UProperty* PropertyThatChanged)
@@ -236,12 +234,11 @@ void SEmitterWidgetDev::NotifyPostChange(const FPropertyChangedEvent& PropertyCh
 	{
 		Emitter->GetProperties()->Init();
 
-		ChildSlot.DetachWidget();
 		BuildContents();
 	}
 
-	//Always refresh the emitter on a propety change.
-	Emitter->Init();
+	//Always refresh the emitter on a property change.
+	Emitter->GetParentEffectInstance()->ReInit();
 }
 
 void SEmitterWidget::Construct(const FArguments& InArgs)
diff --git a/Engine/Source/Editor/UnrealEd/Classes/Niagara/EdGraphSchema_Niagara.h b/Engine/Source/Editor/UnrealEd/Classes/Niagara/EdGraphSchema_Niagara.h
--- a/Engine/Source/Editor/UnrealEd/Classes/Niagara/EdGraphSchema_Niagara.h
+++ b/Engine/Source/Editor/UnrealEd/Classes/Niagara/EdGraphSchema_Niagara.h
@@ -79,6 +79,7 @@ class UEdGraphSchema_Niagara : public UEdGraphSchema
 	static const FLinearColor NodeTitleColor_Constant;
 	static const FLinearColor NodeTitleColor_SystemConstant;
 	static const FLinearColor NodeTitleColor_FunctionCall;
+	static const FLinearColor NodeTitleColor_Event;
 
 private:
 	void GetBreakLinkToSubMenuActions(class FMenuBuilder& MenuBuilder, UEdGraphPin* InGraphPin);
diff --git a/Engine/Source/Editor/UnrealEd/Classes/Niagara/NiagaraGraph.h b/Engine/Source/Editor/UnrealEd/Classes/Niagara/NiagaraGraph.h
--- a/Engine/Source/Editor/UnrealEd/Classes/Niagara/NiagaraGraph.h
+++ b/Engine/Source/Editor/UnrealEd/Classes/Niagara/NiagaraGraph.h
@@ -19,6 +19,9 @@ class UNiagaraGraph : public UEdGraph
 	UNREALED_API class UNiagaraNodeOutput* FindOutputNode() const;
 	UNREALED_API void FindInputNodes(TArray<class UNiagaraNodeInput*>& OutInputNodes) const;
 
+	UNREALED_API void FindReadDataSetNodes(TArray<class UNiagaraNodeReadDataSet*>& OutReadNodes) const;
+	UNREALED_API void FindWriteDataSetNodes(TArray<class UNiagaraNodeWriteDataSet*>& OutWriteNodes) const;
+
 	/** Returns the index of this attribute in the output node of the graph. INDEX_NONE if this is not a valid attribute. */
 	UNREALED_API int32 GetAttributeIndex(const FNiagaraVariableInfo& Attr)const;
 	void UNREALED_API GetAttributes(TArray< FNiagaraVariableInfo >& OutAttributes)const;
diff --git a/Engine/Source/Editor/UnrealEd/Classes/Niagara/NiagaraNode.h b/Engine/Source/Editor/UnrealEd/Classes/Niagara/NiagaraNode.h
--- a/Engine/Source/Editor/UnrealEd/Classes/Niagara/NiagaraNode.h
+++ b/Engine/Source/Editor/UnrealEd/Classes/Niagara/NiagaraNode.h
@@ -7,6 +7,9 @@ UCLASS()
 class UNREALED_API UNiagaraNode : public UEdGraphNode
 {
 	GENERATED_UCLASS_BODY()
+protected:
+
+	void ReallocatePins();
 
 public:
 
diff --git a/Engine/Source/Editor/UnrealEd/Classes/Niagara/NiagaraNodeFunctionCall.h b/Engine/Source/Editor/UnrealEd/Classes/Niagara/NiagaraNodeFunctionCall.h
--- a/Engine/Source/Editor/UnrealEd/Classes/Niagara/NiagaraNodeFunctionCall.h
+++ b/Engine/Source/Editor/UnrealEd/Classes/Niagara/NiagaraNodeFunctionCall.h
@@ -30,10 +30,5 @@ public:
 	virtual FText GetTooltipText() const override;
 	virtual FLinearColor GetNodeTitleColor() const override;
 	// End EdGraphNode interface
-
-private:
-	
-	/** Destroys and reallocates input and output pins. */
-	void ReallocatePins();
 };
 
diff --git a/Engine/Source/Editor/UnrealEd/Classes/Niagara/NiagaraNodeInput.h b/Engine/Source/Editor/UnrealEd/Classes/Niagara/NiagaraNodeInput.h
--- a/Engine/Source/Editor/UnrealEd/Classes/Niagara/NiagaraNodeInput.h
+++ b/Engine/Source/Editor/UnrealEd/Classes/Niagara/NiagaraNodeInput.h
@@ -8,8 +8,6 @@ class UNiagaraNodeInput : public UNiagaraNode
 {
 	GENERATED_UCLASS_BODY()
 
-	void ReallocatePins();
-
 public:
 
 	UPROPERTY(EditAnywhere, Category = Input)
diff --git a/Engine/Source/Editor/UnrealEd/Classes/Niagara/NiagaraNodeOutput.h b/Engine/Source/Editor/UnrealEd/Classes/Niagara/NiagaraNodeOutput.h
--- a/Engine/Source/Editor/UnrealEd/Classes/Niagara/NiagaraNodeOutput.h
+++ b/Engine/Source/Editor/UnrealEd/Classes/Niagara/NiagaraNodeOutput.h
@@ -8,12 +8,11 @@ class UNiagaraNodeOutput : public UNiagaraNode
 {
 	GENERATED_UCLASS_BODY()
 
-	void ReallocatePins();
-public:
-
 	UPROPERTY(EditAnywhere, Category = Output)
 	TArray<FNiagaraVariableInfo> Outputs;
 
+public:
+
 	// Begin UObject interface
 	virtual void PostEditChangeProperty(struct FPropertyChangedEvent& PropertyChangedEvent) override;
 	// End UObject interface
diff --git a/Engine/Source/Editor/UnrealEd/Classes/Niagara/NiagaraNodeReadDataSet.h b/Engine/Source/Editor/UnrealEd/Classes/Niagara/NiagaraNodeReadDataSet.h
new file mode 100644
--- /dev/null
+++ b/Engine/Source/Editor/UnrealEd/Classes/Niagara/NiagaraNodeReadDataSet.h
@@ -0,0 +1,31 @@
+// Copyright 1998-2014 Epic Games, Inc. All Rights Reserved.
+
+#pragma once
+#include "NiagaraNodeReadDataSet.generated.h"
+
+UCLASS(MinimalAPI)
+class UNiagaraNodeReadDataSet : public UNiagaraNode
+{
+	GENERATED_UCLASS_BODY()
+
+public:
+
+	UPROPERTY(EditAnywhere, Category = DataSet)
+	FNiagaraDataSetID DataSet;
+
+	UPROPERTY(EditAnywhere, Category = Variables)
+	TArray<FNiagaraVariableInfo> Variables;
+
+	// Begin UObject interface
+	virtual void PostEditChangeProperty(struct FPropertyChangedEvent& PropertyChangedEvent) override;
+	// End UObject interface
+
+	// Begin EdGraphNode interface
+	virtual void AllocateDefaultPins() override;
+	virtual FText GetNodeTitle(ENodeTitleType::Type TitleType) const override;
+	virtual FLinearColor GetNodeTitleColor() const override;
+	// End EdGraphNode interface
+
+	virtual void Compile(class INiagaraCompiler* Compiler, TArray<FNiagaraNodeResult>& Outputs)override;
+};
+
diff --git a/Engine/Source/Editor/UnrealEd/Classes/Niagara/NiagaraNodeWriteDataSet.h b/Engine/Source/Editor/UnrealEd/Classes/Niagara/NiagaraNodeWriteDataSet.h
new file mode 100644
--- /dev/null
+++ b/Engine/Source/Editor/UnrealEd/Classes/Niagara/NiagaraNodeWriteDataSet.h
@@ -0,0 +1,31 @@
+// Copyright 1998-2014 Epic Games, Inc. All Rights Reserved.
+
+#pragma once
+#include "NiagaraNodeWriteDataSet.generated.h"
+
+UCLASS(MinimalAPI)
+class UNiagaraNodeWriteDataSet : public UNiagaraNode
+{
+	GENERATED_UCLASS_BODY()
+
+public:
+
+	UPROPERTY(EditAnywhere, Category = DataSet)
+	FNiagaraDataSetID DataSet;
+
+	UPROPERTY(EditAnywhere, Category = Variables)
+	TArray<FNiagaraVariableInfo> Variables;
+
+	// Begin UObject interface
+	virtual void PostEditChangeProperty(struct FPropertyChangedEvent& PropertyChangedEvent) override;
+	// End UObject interface
+
+	// Begin EdGraphNode interface
+	virtual void AllocateDefaultPins() override;
+	virtual FText GetNodeTitle(ENodeTitleType::Type TitleType) const override;
+	virtual FLinearColor GetNodeTitleColor() const override;
+	// End EdGraphNode interface
+
+	virtual void Compile(class INiagaraCompiler* Compiler, TArray<FNiagaraNodeResult>& Outputs)override;
+};
+
diff --git a/Engine/Source/Editor/UnrealEd/Classes/Niagara/NiagaraScriptSource.h b/Engine/Source/Editor/UnrealEd/Classes/Niagara/NiagaraScriptSource.h
--- a/Engine/Source/Editor/UnrealEd/Classes/Niagara/NiagaraScriptSource.h
+++ b/Engine/Source/Editor/UnrealEd/Classes/Niagara/NiagaraScriptSource.h
@@ -15,6 +15,10 @@ class UNiagaraScriptSource : public UNiagaraScriptSourceBase
 	UPROPERTY()
 	class UNiagaraGraph*	NodeGraph;
 
+	/** The same node graph from above but with all function calls merge into a single graph. */
+	UPROPERTY()
+	class UNiagaraGraph*	FlattenedNodeGraph;
+
 	// UObject interface.
 	virtual void PostLoad() override;
 	virtual void Compile() override;
diff --git a/Engine/Source/Editor/UnrealEd/Private/Factories/NiagaraScriptFactory.cpp b/Engine/Source/Editor/UnrealEd/Private/Factories/NiagaraScriptFactory.cpp
--- a/Engine/Source/Editor/UnrealEd/Private/Factories/NiagaraScriptFactory.cpp
+++ b/Engine/Source/Editor/UnrealEd/Private/Factories/NiagaraScriptFactory.cpp
@@ -41,6 +41,8 @@ UObject* UNiagaraScriptFactoryNew::FactoryCreateNew(UClass* Class,UObject* InPar
 
 			//Automatically add attributes that the current render modules rely on. Not sure what we'll want to do here in the end.
 			OutputNode->Outputs.Add(FNiagaraVariableInfo(FName(TEXT("Position")), ENiagaraDataType::Vector));
+			OutputNode->Outputs.Add(FNiagaraVariableInfo(FName(TEXT("Velocity")), ENiagaraDataType::Vector));
+			OutputNode->Outputs.Add(FNiagaraVariableInfo(FName(TEXT("Acceleration")), ENiagaraDataType::Vector));
 			OutputNode->Outputs.Add(FNiagaraVariableInfo(FName(TEXT("Rotation")), ENiagaraDataType::Vector));
 			OutputNode->Outputs.Add(FNiagaraVariableInfo(FName(TEXT("Color")), ENiagaraDataType::Vector));
 			OutputNode->Outputs.Add(FNiagaraVariableInfo(FName(TEXT("Age")), ENiagaraDataType::Vector));
diff --git a/Engine/Source/Editor/UnrealEd/Private/Niagara/EdGraphSchema_Niagara.cpp b/Engine/Source/Editor/UnrealEd/Private/Niagara/EdGraphSchema_Niagara.cpp
--- a/Engine/Source/Editor/UnrealEd/Private/Niagara/EdGraphSchema_Niagara.cpp
+++ b/Engine/Source/Editor/UnrealEd/Private/Niagara/EdGraphSchema_Niagara.cpp
@@ -19,6 +19,7 @@ const FLinearColor UEdGraphSchema_Niagara::NodeTitleColor_Attribute = FLinearCol
 const FLinearColor UEdGraphSchema_Niagara::NodeTitleColor_Constant = FLinearColor::Red;
 const FLinearColor UEdGraphSchema_Niagara::NodeTitleColor_SystemConstant = FLinearColor::White;
 const FLinearColor UEdGraphSchema_Niagara::NodeTitleColor_FunctionCall = FLinearColor::Blue;
+const FLinearColor UEdGraphSchema_Niagara::NodeTitleColor_Event = FLinearColor::Red;
 
 namespace 
 {
@@ -200,6 +201,24 @@ void UEdGraphSchema_Niagara::GetGraphContextActions(FGraphContextMenuBuilder& Co
 		UNiagaraNodeFunctionCall* FunctionCallNode = NewObject<UNiagaraNodeFunctionCall>(ContextMenuBuilder.OwnerOfTemporaries);
 		FunctionCallAction->NodeTemplate = FunctionCallNode;
 	}
+
+	//Add event read and writes nodes
+	{
+		const FText MenuDesc = LOCTEXT("NiagaraEventRead", "Event Read");
+
+		TSharedPtr<FNiagaraSchemaAction_NewNode> Action = AddNewNodeAction(ContextMenuBuilder, FText::GetEmpty(), MenuDesc, TEXT(""));
+
+		UNiagaraNodeReadDataSet* EventReadNode = NewObject<UNiagaraNodeReadDataSet>(ContextMenuBuilder.OwnerOfTemporaries);
+		Action->NodeTemplate = EventReadNode;
+	}	
+	{
+		const FText MenuDesc = LOCTEXT("NiagaraEventWrite", "Event Write");
+
+		TSharedPtr<FNiagaraSchemaAction_NewNode> Action = AddNewNodeAction(ContextMenuBuilder, FText::GetEmpty(), MenuDesc, TEXT(""));
+
+		UNiagaraNodeWriteDataSet* EventWriteNode = NewObject<UNiagaraNodeWriteDataSet>(ContextMenuBuilder.OwnerOfTemporaries);
+		Action->NodeTemplate = EventWriteNode;
+	}
 	//TODO: Add quick commands for certain UNiagaraStructs and UNiagaraScripts to be added as functions
 }
 
diff --git a/Engine/Source/Editor/UnrealEd/Private/Niagara/NiagaraEditorCommon.cpp b/Engine/Source/Editor/UnrealEd/Private/Niagara/NiagaraEditorCommon.cpp
--- a/Engine/Source/Editor/UnrealEd/Private/Niagara/NiagaraEditorCommon.cpp
+++ b/Engine/Source/Editor/UnrealEd/Private/Niagara/NiagaraEditorCommon.cpp
@@ -21,6 +21,9 @@ void FNiagaraOpInfo::Init()
 	FName Y(TEXT("Y"));	FText YText = NSLOCTEXT("NiagaraOpInfo", "Second Vector Component", "Y");
 	FName Z(TEXT("Z"));	FText ZText = NSLOCTEXT("NiagaraOpInfo", "Third Vector Component", "Z");
 	FName W(TEXT("W"));	FText WText = NSLOCTEXT("NiagaraOpInfo", "Fourth Vector Component", "W");
+
+	FName Selection(TEXT("Selection"));	FText SelectionText = NSLOCTEXT("NiagaraOpInfo", "Selection", "Selection");
+	FName Mask(TEXT("Mask"));	FText MaskText = NSLOCTEXT("NiagaraOpInfo", "Mask", "Mask");
 	
 	FText MinText = NSLOCTEXT("NiagaraOpInfo", "Min", "Min");
 	FText MaxText = NSLOCTEXT("NiagaraOpInfo", "Max", "Max");
@@ -516,7 +519,6 @@ void FNiagaraOpInfo::Init()
 	Op->Outputs.Add(FNiagaraOpInOutInfo(Result, ENiagaraDataType::Matrix, ResultText, ResultText, DefaultStr_VecOne));
 	Op->OpDelegate.BindStatic(&INiagaraCompiler::Inverse);
 
-
 	Op = &OpInfos.Add(LessThan);
 	Op->Name = LessThan;
 	Op->FriendlyName = NSLOCTEXT("NiagaraOpInfo", "LessThan Name", "A Less Than B");
@@ -526,6 +528,25 @@ void FNiagaraOpInfo::Init()
 	Op->Outputs.Add(FNiagaraOpInOutInfo(Result, ENiagaraDataType::Vector, ResultText, ResultText, DefaultStr_VecOne));
 	Op->OpDelegate.BindStatic(&INiagaraCompiler::LessThan);
 
+	Op = &OpInfos.Add(GreaterThan);
+	Op->Name = GreaterThan;
+	Op->FriendlyName = NSLOCTEXT("NiagaraOpInfo", "GreaterThan Name", "A Greater Than B");
+	Op->Description = NSLOCTEXT("NiagaraOpInfo", "GreaterThan Desc", "Result = 1 if A>B, 0 otherwise");
+	Op->Inputs.Add(FNiagaraOpInOutInfo(A, ENiagaraDataType::Vector, AText, AText, DefaultStr_VecZero));
+	Op->Inputs.Add(FNiagaraOpInOutInfo(B, ENiagaraDataType::Vector, BText, BText, DefaultStr_VecOne));
+	Op->Outputs.Add(FNiagaraOpInOutInfo(Result, ENiagaraDataType::Vector, ResultText, ResultText, DefaultStr_VecOne));
+	Op->OpDelegate.BindStatic(&INiagaraCompiler::GreaterThan);
+
+	Op = &OpInfos.Add(Select);
+	Op->Name = Select;
+	Op->FriendlyName = NSLOCTEXT("NiagaraOpInfo", "Select Name", "Select");
+	Op->Description = NSLOCTEXT("NiagaraOpInfo", "Select Desc", "Selects A or B depending on Selection. Selection > 0.0f = A, otherwise B.");
+	Op->Inputs.Add(FNiagaraOpInOutInfo(Selection, ENiagaraDataType::Vector, SelectionText, SelectionText, DefaultStr_VecOne));
+	Op->Inputs.Add(FNiagaraOpInOutInfo(A, ENiagaraDataType::Vector, AText, AText, DefaultStr_VecZero));
+	Op->Inputs.Add(FNiagaraOpInOutInfo(B, ENiagaraDataType::Vector, BText, BText, DefaultStr_VecOne));
+	Op->Outputs.Add(FNiagaraOpInOutInfo(Result, ENiagaraDataType::Vector, ResultText, ResultText, DefaultStr_VecOne));
+	Op->OpDelegate.BindStatic(&INiagaraCompiler::Select);
+
 	Op = &OpInfos.Add(Sample);
 	Op->Name = Sample;
 	Op->FriendlyName = NSLOCTEXT("NiagaraOpInfo", "Sample Name", "Buffer Sample");
@@ -545,14 +566,6 @@ void FNiagaraOpInfo::Init()
 	Op->Outputs.Add(FNiagaraOpInOutInfo(Result, ENiagaraDataType::Vector, ResultText, ResultText, DefaultStr_VecOne));
 	Op->OpDelegate.BindStatic(&INiagaraCompiler::Write);
 
-	Op = &OpInfos.Add(EventBroadcast);
-	Op->Name = EventBroadcast;
-	Op->FriendlyName = NSLOCTEXT("NiagaraOpInfo", "EventBroadcast Name", "Broadcast Event");
-	Op->Description = NSLOCTEXT("NiagaraOpInfo", "EventBroadcast Desc", "Broadcast Event if trigger is >=1.0");
-	Op->Inputs.Add(FNiagaraOpInOutInfo(A, ENiagaraDataType::Vector, NSLOCTEXT("NiagaraOpInfo", "Trigger", "Trigger"), NSLOCTEXT("NiagaraOpInfo", "Trigger", "Trigger"), DefaultStr_VecOne));
-//	Op->Outputs.Add(FNiagaraOpInOutInfo(Result, ENiagaraDataType::Vector, ResultText, ResultText, DefaultStr_VecOne));
-	Op->OpDelegate.BindStatic(&INiagaraCompiler::EventBroadcast);
-
 	Op = &OpInfos.Add(EaseIn);
 	Op->Name = EaseIn;
 	Op->FriendlyName = NSLOCTEXT("NiagaraOpInfo", "EaseIn Name", "Ease In 0->1");
diff --git a/Engine/Source/Editor/UnrealEd/Private/Niagara/NiagaraNode.cpp b/Engine/Source/Editor/UnrealEd/Private/Niagara/NiagaraNode.cpp
--- a/Engine/Source/Editor/UnrealEd/Private/Niagara/NiagaraNode.cpp
+++ b/Engine/Source/Editor/UnrealEd/Private/Niagara/NiagaraNode.cpp
@@ -7,6 +7,34 @@ UNiagaraNode::UNiagaraNode(const FObjectInitializer& ObjectInitializer)
 {
 }
 
+void UNiagaraNode::ReallocatePins()
+{
+	Modify();
+
+	// Move the existing pins to a saved array
+	TArray<UEdGraphPin*> OldPins(Pins);
+	Pins.Empty();
+
+	// Recreate the new pins
+	AllocateDefaultPins();
+
+	// Copy the old pin data and remove it.
+	for (int32 OldPinIndex = 0; OldPinIndex < OldPins.Num(); ++OldPinIndex)
+	{
+		UEdGraphPin* OldPin = OldPins[OldPinIndex];
+		if (UEdGraphPin** MatchingNewPin = Pins.FindByPredicate([&](UEdGraphPin* Pin){ return Pin->Direction == OldPin->Direction && Pin->PinName == OldPin->PinName; }))
+		{
+			if (*MatchingNewPin && OldPin)
+				(*MatchingNewPin)->CopyPersistentDataFromOldPin(*OldPin);
+		}
+		OldPin->Modify();
+		OldPin->BreakAllPinLinks();
+	}
+	OldPins.Empty();
+
+	GetGraph()->NotifyGraphChanged();
+}
+
 void UNiagaraNode::AutowireNewNode(UEdGraphPin* FromPin)
 {
 	Super::AutowireNewNode(FromPin);
@@ -117,4 +145,4 @@ void UNiagaraNode::GetOutputPins(TArray<class UEdGraphPin*>& OutOutputPins) cons
 			OutOutputPins.Add(Pins[PinIndex]);
 		}
 	}
-}
+}
\ No newline at end of file
diff --git a/Engine/Source/Editor/UnrealEd/Private/Niagara/NiagaraNodeInput.cpp b/Engine/Source/Editor/UnrealEd/Private/Niagara/NiagaraNodeInput.cpp
--- a/Engine/Source/Editor/UnrealEd/Private/Niagara/NiagaraNodeInput.cpp
+++ b/Engine/Source/Editor/UnrealEd/Private/Niagara/NiagaraNodeInput.cpp
@@ -26,115 +26,37 @@ void UNiagaraNodeInput::PostEditChangeProperty(struct FPropertyChangedEvent& Pro
 	Super::PostEditChangeProperty(PropertyChangedEvent);
 }
 
-void UNiagaraNodeInput::ReallocatePins()
+void UNiagaraNodeInput::AllocateDefaultPins()
 {
 	const UEdGraphSchema_Niagara* Schema = GetDefault<UEdGraphSchema_Niagara>();
-	Modify();
-
-	// Break any links to 'orphan' pins
-	for (int32 PinIndex = 0; PinIndex < Pins.Num(); ++PinIndex)
+	switch (Input.Type)
 	{
-		UEdGraphPin* Pin = Pins[PinIndex];
-		TArray<class UEdGraphPin*>& LinkedToRef = Pin->LinkedTo;
-		for (int32 LinkIdx = 0; LinkIdx < LinkedToRef.Num(); LinkIdx++)
+		case ENiagaraDataType::Scalar:
 		{
-			UEdGraphPin* OtherPin = LinkedToRef[LinkIdx];
-			// If we are linked to a pin that its owner doesn't know about, break that link
-			if (!OtherPin->GetOwningNode()->Pins.Contains(OtherPin))
-			{
-				Pin->LinkedTo.Remove(OtherPin);
-			}
+			CreatePin(EGPD_Output, Schema->PC_Float, TEXT(""), NULL, false, false, TEXT("Input"));
 		}
-	}
-
-	// Store the old Input and Output pins
-	TArray<UEdGraphPin*> OldInputPins;
-	TArray<UEdGraphPin*> OldOutputPins;
-	GetInputPins(OldInputPins);
-	GetOutputPins(OldOutputPins);
-
-	// Move the existing pins to a saved array
-	TArray<UEdGraphPin*> OldPins(Pins);
-	Pins.Empty();
-
-
-	switch (Input.Type)
-	{
-	case ENiagaraDataType::Scalar:
-	{
-		CreatePin(EGPD_Output, Schema->PC_Float, TEXT(""), NULL, false, false, Input.Name.ToString());
-	}
-		break;
-	case ENiagaraDataType::Vector:
-	{
-		CreatePin(EGPD_Output, Schema->PC_Vector, TEXT(""), NULL, false, false, Input.Name.ToString());
-	}
-		break;
-	case ENiagaraDataType::Matrix:
-	{
-		CreatePin(EGPD_Output, Schema->PC_Matrix, TEXT(""), NULL, false, false, Input.Name.ToString());
-	}
-		break;
-	case ENiagaraDataType::Curve:
-	{
-		CreatePin(EGPD_Output, Schema->PC_Curve, TEXT(""), NULL, false, false, Input.Name.ToString());
-	}
-		break;
-	};
-
-	// Get new Input and Output pins
-	TArray<UEdGraphPin*> NewInputPins;
-	TArray<UEdGraphPin*> NewOutputPins;
-	GetInputPins(NewInputPins);
-	GetOutputPins(NewOutputPins);
-
-	for (int32 PinIndex = 0; PinIndex < OldInputPins.Num(); PinIndex++)
-	{
-		if (PinIndex < NewInputPins.Num())
+			break;
+		case ENiagaraDataType::Vector:
 		{
-			NewInputPins[PinIndex]->CopyPersistentDataFromOldPin(*OldInputPins[PinIndex]);
+			CreatePin(EGPD_Output, Schema->PC_Vector, TEXT(""), NULL, false, false, TEXT("Input"));
 		}
-	}
-
-	for (int32 PinIndex = 0; PinIndex < OldOutputPins.Num(); PinIndex++)
-	{
-		if (PinIndex < NewOutputPins.Num())
+			break;
+		case ENiagaraDataType::Matrix:
 		{
-			NewOutputPins[PinIndex]->CopyPersistentDataFromOldPin(*OldOutputPins[PinIndex]);
+			CreatePin(EGPD_Output, Schema->PC_Matrix, TEXT(""), NULL, false, false, TEXT("Input"));
 		}
-	}
-
-	OldInputPins.Empty();
-	OldOutputPins.Empty();
-
-	// Throw away the original pins
-	for (int32 OldPinIndex = 0; OldPinIndex < OldPins.Num(); ++OldPinIndex)
-	{
-		UEdGraphPin* OldPin = OldPins[OldPinIndex];
-		OldPin->Modify();
-		OldPin->BreakAllPinLinks();
-
-#if 0
-		UEdGraphNode::ReturnPinToPool(OldPin);
-#else
-		OldPin->Rename(NULL, GetTransientPackage(), REN_None);
-		OldPin->RemoveFromRoot();
-		OldPin->MarkPendingKill();
-#endif
-	}
-	OldPins.Empty();
-
-	GetGraph()->NotifyGraphChanged();
-}
-
-void UNiagaraNodeInput::AllocateDefaultPins()
-{
-	ReallocatePins();
+			break;
+		case ENiagaraDataType::Curve:
+		{
+			CreatePin(EGPD_Output, Schema->PC_Curve, TEXT(""), NULL, false, false, TEXT("Input"));
+		}
+		break;
+	};
 }
 
 FText UNiagaraNodeInput::GetNodeTitle(ENodeTitleType::Type TitleType) const
 {
-	return NSLOCTEXT("NiagaraNodeInput", "Input", "Input");
+	return FText::FromName(Input.Name);
 }
 
 FLinearColor UNiagaraNodeInput::GetNodeTitleColor() const
diff --git a/Engine/Source/Editor/UnrealEd/Private/Niagara/NiagaraNodeOutput.cpp b/Engine/Source/Editor/UnrealEd/Private/Niagara/NiagaraNodeOutput.cpp
--- a/Engine/Source/Editor/UnrealEd/Private/Niagara/NiagaraNodeOutput.cpp
+++ b/Engine/Source/Editor/UnrealEd/Private/Niagara/NiagaraNodeOutput.cpp
@@ -20,84 +20,6 @@ void UNiagaraNodeOutput::PostEditChangeProperty(struct FPropertyChangedEvent& Pr
 	Super::PostEditChangeProperty(PropertyChangedEvent);
 }
 
-void UNiagaraNodeOutput::ReallocatePins()
-{
-	Modify();
-
-	// Break any links to 'orphan' pins
-	for (int32 PinIndex = 0; PinIndex < Pins.Num(); ++PinIndex)
-	{
-		UEdGraphPin* Pin = Pins[PinIndex];
-		TArray<class UEdGraphPin*>& LinkedToRef = Pin->LinkedTo;
-		for (int32 LinkIdx = 0; LinkIdx < LinkedToRef.Num(); LinkIdx++)
-		{
-			UEdGraphPin* OtherPin = LinkedToRef[LinkIdx];
-			// If we are linked to a pin that its owner doesn't know about, break that link
-			if (!OtherPin->GetOwningNode()->Pins.Contains(OtherPin))
-			{
-				Pin->LinkedTo.Remove(OtherPin);
-			}
-		}
-	}
-
-	// Store the old Input and Output pins
-	TArray<UEdGraphPin*> OldInputPins;
-	TArray<UEdGraphPin*> OldOutputPins;
-	GetInputPins(OldInputPins);
-	GetOutputPins(OldOutputPins);
-
-	// Move the existing pins to a saved array
-	TArray<UEdGraphPin*> OldPins(Pins);
-	Pins.Empty();
-
-	// Recreate the new pins
-	AllocateDefaultPins();
-
-	// Get new Input and Output pins
-	TArray<UEdGraphPin*> NewInputPins;
-	TArray<UEdGraphPin*> NewOutputPins;
-	GetInputPins(NewInputPins);
-	GetOutputPins(NewOutputPins);
-
-	for (int32 PinIndex = 0; PinIndex < OldInputPins.Num(); PinIndex++)
-	{
-		if (PinIndex < NewInputPins.Num())
-		{
-			NewInputPins[PinIndex]->CopyPersistentDataFromOldPin(*OldInputPins[PinIndex]);
-		}
-	}
-
-	for (int32 PinIndex = 0; PinIndex < OldOutputPins.Num(); PinIndex++)
-	{
-		if (PinIndex < NewOutputPins.Num())
-		{
-			NewOutputPins[PinIndex]->CopyPersistentDataFromOldPin(*OldOutputPins[PinIndex]);
-		}
-	}
-
-	OldInputPins.Empty();
-	OldOutputPins.Empty();
-
-	// Throw away the original pins
-	for (int32 OldPinIndex = 0; OldPinIndex < OldPins.Num(); ++OldPinIndex)
-	{
-		UEdGraphPin* OldPin = OldPins[OldPinIndex];
-		OldPin->Modify();
-		OldPin->BreakAllPinLinks();
-
-#if 0
-		UEdGraphNode::ReturnPinToPool(OldPin);
-#else
-		OldPin->Rename(NULL, GetTransientPackage(), REN_None);
-		OldPin->RemoveFromRoot();
-		OldPin->MarkPendingKill();
-#endif
-	}
-	OldPins.Empty();
-
-	GetGraph()->NotifyGraphChanged();
-}
-
 void UNiagaraNodeOutput::AllocateDefaultPins()
 {
 	const UEdGraphSchema_Niagara* Schema = GetDefault<UEdGraphSchema_Niagara>();
diff --git a/Engine/Source/Editor/UnrealEd/Private/Niagara/NiagaraNodeReadDataSet.cpp b/Engine/Source/Editor/UnrealEd/Private/Niagara/NiagaraNodeReadDataSet.cpp
new file mode 100644
--- /dev/null
+++ b/Engine/Source/Editor/UnrealEd/Private/Niagara/NiagaraNodeReadDataSet.cpp
@@ -0,0 +1,96 @@
+// Copyright 1998-2014 Epic Games, Inc. All Rights Reserved.
+
+#include "UnrealEd.h"
+#include "BlueprintGraphDefinitions.h"
+#include "GraphEditorSettings.h"
+
+#define LOCTEXT_NAMESPACE "NiagaraNodeDataSetRead"
+
+UNiagaraNodeReadDataSet::UNiagaraNodeReadDataSet(const FObjectInitializer& ObjectInitializer)
+: Super(ObjectInitializer)
+{
+}
+
+void UNiagaraNodeReadDataSet::PostEditChangeProperty(struct FPropertyChangedEvent& PropertyChangedEvent)
+{
+	if (PropertyChangedEvent.Property != nullptr)
+	{
+		ReallocatePins();
+	}
+	Super::PostEditChangeProperty(PropertyChangedEvent);
+}
+
+void UNiagaraNodeReadDataSet::AllocateDefaultPins()
+{
+	const UEdGraphSchema_Niagara* Schema = GetDefault<UEdGraphSchema_Niagara>();
+
+	if (DataSet.Type == ENiagaraDataSetType::Event)
+	{
+		//Probably need this for all data set types tbh!
+		UEdGraphPin* Pin = CreatePin(EGPD_Output, Schema->PC_Vector, TEXT(""), NULL, false, false, TEXT("Valid"));//TODO - CHANGE TO INT / BOOL / SCALAR
+		Pin->bDefaultValueIsIgnored = true;
+	}
+
+	for (const FNiagaraVariableInfo& Var : Variables)
+	{
+		switch (Var.Type)
+		{
+		case ENiagaraDataType::Scalar:
+		{
+			CreatePin(EGPD_Output, Schema->PC_Float, TEXT(""), NULL, false, false, Var.Name.ToString());
+		}
+			break;
+		case ENiagaraDataType::Vector:
+		{
+			CreatePin(EGPD_Output, Schema->PC_Vector, TEXT(""), NULL, false, false, Var.Name.ToString());
+		}
+			break;
+		case ENiagaraDataType::Matrix:
+		{
+			CreatePin(EGPD_Output, Schema->PC_Matrix, TEXT(""), NULL, false, false, Var.Name.ToString());
+		}
+			break;
+		case ENiagaraDataType::Curve:
+		{
+			CreatePin(EGPD_Output, Schema->PC_Curve, TEXT(""), NULL, false, false, Var.Name.ToString());
+		}
+			break;
+		};
+	}
+}
+
+FText UNiagaraNodeReadDataSet::GetNodeTitle(ENodeTitleType::Type TitleType) const
+{
+	return FText::Format(LOCTEXT("NiagaraDataSetReadFormat", "{0} Read"), FText::FromName(DataSet.Name));
+}
+
+FLinearColor UNiagaraNodeReadDataSet::GetNodeTitleColor() const
+{
+	check(DataSet.Type == ENiagaraDataSetType::Event);//Implement other datasets
+	return CastChecked<UEdGraphSchema_Niagara>(GetSchema())->NodeTitleColor_Event;
+}
+
+void UNiagaraNodeReadDataSet::Compile(class INiagaraCompiler* Compiler, TArray<FNiagaraNodeResult>& Outputs)
+{
+	if(DataSet.Type == ENiagaraDataSetType::Event)
+	{
+		TNiagaraExprPtr IndexExpr = Compiler->ConsumeSharedDataIndex(DataSet, true);//Where do we decide wrapping?
+		check(Variables.Num() == Pins.Num() - 1);//We should have a pin for each variable + 1 for the Valid pin.
+		Outputs.Add(FNiagaraNodeResult(Compiler->SharedDataIndexIsValid(DataSet, IndexExpr), Pins[0]));
+		for (int32 i = 0; i < Variables.Num(); ++i)
+		{
+			check(Variables[i].Type == ENiagaraDataType::Vector);//Only support vectors currently.
+			Outputs.Add(FNiagaraNodeResult(Compiler->SharedDataRead(DataSet, Variables[i], IndexExpr), Pins[i+1]));
+		}
+	}
+	else
+	{
+		Compiler->Error(LOCTEXT("UnImplDataSetTypeError", "This data set type hasn't been implemented."), this, nullptr);
+	}
+	
+}
+
+#undef LOCTEXT_NAMESPACE
+
+
+
diff --git a/Engine/Source/Editor/UnrealEd/Private/Niagara/NiagaraNodeWriteDataSet.cpp b/Engine/Source/Editor/UnrealEd/Private/Niagara/NiagaraNodeWriteDataSet.cpp
new file mode 100644
--- /dev/null
+++ b/Engine/Source/Editor/UnrealEd/Private/Niagara/NiagaraNodeWriteDataSet.cpp
@@ -0,0 +1,135 @@
+// Copyright 1998-2014 Epic Games, Inc. All Rights Reserved.
+
+#include "UnrealEd.h"
+#include "BlueprintGraphDefinitions.h"
+#include "GraphEditorSettings.h"
+
+#define LOCTEXT_NAMESPACE "NiagaraNodeWriteDataSet"
+
+UNiagaraNodeWriteDataSet::UNiagaraNodeWriteDataSet(const FObjectInitializer& ObjectInitializer)
+: Super(ObjectInitializer)
+{
+}
+
+void UNiagaraNodeWriteDataSet::PostEditChangeProperty(struct FPropertyChangedEvent& PropertyChangedEvent)
+{
+	if (PropertyChangedEvent.Property != nullptr)
+	{
+		ReallocatePins();
+	}
+	Super::PostEditChangeProperty(PropertyChangedEvent);
+}
+
+void UNiagaraNodeWriteDataSet::AllocateDefaultPins()
+{
+	const UEdGraphSchema_Niagara* Schema = GetDefault<UEdGraphSchema_Niagara>();
+
+	if (DataSet.Type == ENiagaraDataSetType::Event)
+	{
+		UEdGraphPin* Pin = CreatePin(EGPD_Input, Schema->PC_Vector, TEXT(""), NULL, false, false, TEXT("Valid"));//TODO - CHANGE TO INT / BOOL / SCALAR
+		Pin->bDefaultValueIsIgnored = true;
+	}
+	
+	for (const FNiagaraVariableInfo& Var : Variables)
+	{
+		switch (Var.Type)
+		{
+		case ENiagaraDataType::Scalar:
+		{
+			CreatePin(EGPD_Input, Schema->PC_Float, TEXT(""), NULL, false, false, Var.Name.ToString());
+		}
+			break;
+		case ENiagaraDataType::Vector:
+		{
+			CreatePin(EGPD_Input, Schema->PC_Vector, TEXT(""), NULL, false, false, Var.Name.ToString());
+		}
+			break;
+		case ENiagaraDataType::Matrix:
+		{
+			CreatePin(EGPD_Input, Schema->PC_Matrix, TEXT(""), NULL, false, false, Var.Name.ToString());
+		}
+			break;
+		case ENiagaraDataType::Curve:
+		{
+			CreatePin(EGPD_Input, Schema->PC_Curve, TEXT(""), NULL, false, false, Var.Name.ToString());
+		}
+			break;
+		};
+	}
+}
+
+FText UNiagaraNodeWriteDataSet::GetNodeTitle(ENodeTitleType::Type TitleType) const
+{
+	return FText::Format(LOCTEXT("NiagaraDataSetWriteFormat", "{0} Write"), FText::FromName(DataSet.Name));
+}
+
+FLinearColor UNiagaraNodeWriteDataSet::GetNodeTitleColor() const
+{
+	check(DataSet.Type == ENiagaraDataSetType::Event);//Implement other datasets
+	return CastChecked<UEdGraphSchema_Niagara>(GetSchema())->NodeTitleColor_Event;
+}
+
+void UNiagaraNodeWriteDataSet::Compile(class INiagaraCompiler* Compiler, TArray<FNiagaraNodeResult>& Outputs)
+{
+	bool bError = false;
+
+	if (DataSet.Type == ENiagaraDataSetType::Event)
+	{
+		//Compile the Emit pin.
+		TNiagaraExprPtr EmitExpression = Compiler->CompilePin(Pins[0]);
+
+		if (EmitExpression.IsValid())
+		{
+			//Test the Valid pin result against 0. Maybe just require a direct connection of 0 or 0xFFFFFFFF?
+			FNiagaraVariableInfo Zero(TEXT("0.0, 0.0, 0.0, 0.0"), ENiagaraDataType::Vector);
+			TNiagaraExprPtr ZeroContantExpression = Compiler->GetInternalConstant(Zero, FVector4(0.0f, 0.0f, 0.0f, 0.0f));
+			TArray<TNiagaraExprPtr> ConditonInputs;
+			ConditonInputs.Add(EmitExpression);
+			ConditonInputs.Add(ZeroContantExpression);
+			check(ZeroContantExpression.IsValid());
+			TArray<TNiagaraExprPtr> ConditionOpOutputs;
+			INiagaraCompiler::GreaterThan(Compiler, ConditonInputs, ConditionOpOutputs);
+			TNiagaraExprPtr ValidExpr = ConditionOpOutputs[0];
+
+			TArray<TNiagaraExprPtr> InputExpressions;
+			for (int32 i = 0; i < Variables.Num(); ++i)
+			{
+				const FNiagaraVariableInfo& Var = Variables[i];
+				UEdGraphPin* Pin = Pins[i + 1];//Pin[0] is emit
+				check(Pin->Direction == EGPD_Input);
+
+				TNiagaraExprPtr Result = Compiler->CompilePin(Pin);
+				if (!Result.IsValid())
+				{
+					bError = true;
+					Compiler->Error(FText::Format(LOCTEXT("DataSetWriteErrorFormat", "Error writing variable {0} to dataset {1}"), FText::FromName(DataSet.Name), Pin->PinFriendlyName), this, Pin);
+				}
+
+				InputExpressions.Add(Result);
+			}
+
+			//Gets the index to write to. 
+			TNiagaraExprPtr IndexExpression = Compiler->AcquireSharedDataIndex(DataSet, true, ValidExpr);
+
+			if (!bError)
+			{
+				check(Variables.Num() == InputExpressions.Num());
+				for (int32 i = 0; i < Variables.Num(); ++i)
+				{
+					Outputs.Add(FNiagaraNodeResult(Compiler->SharedDataWrite(DataSet, Variables[i], IndexExpression, InputExpressions[i]), Pins[i + 1]));//Pin[0] is Emit
+				}
+			}
+		}
+	}
+	else
+	{
+		check(false);//IMPLEMENT OTHER DATA SETS.
+	}
+}
+
+#undef LOCTEXT_NAMESPACE
+
+
+
+
+
diff --git a/Engine/Source/Editor/UnrealEd/Private/Niagara/NiagaraScriptSource.cpp b/Engine/Source/Editor/UnrealEd/Private/Niagara/NiagaraScriptSource.cpp
--- a/Engine/Source/Editor/UnrealEd/Private/Niagara/NiagaraScriptSource.cpp
+++ b/Engine/Source/Editor/UnrealEd/Private/Niagara/NiagaraScriptSource.cpp
@@ -4,6 +4,10 @@
 #include "NiagaraScript.h"
 #include "NiagaraConstants.h"
 #include "NiagaraEditorModule.h"
+#include "ComponentReregisterContext.h"
+#include "NiagaraComponent.h"
+#include "NiagaraSimulation.h"
+#include "NiagaraEffect.h"
 
 //////////////////////////////////////////////////////////////////////////
 // NiagaraGraph
@@ -49,6 +53,28 @@ void UNiagaraGraph::FindInputNodes(TArray<class UNiagaraNodeInput*>& OutInputNod
 	}
 }
 
+void UNiagaraGraph::FindReadDataSetNodes(TArray<class UNiagaraNodeReadDataSet*>& OutReadNodes) const
+{
+	for (UEdGraphNode* Node : Nodes)
+	{
+		if (UNiagaraNodeReadDataSet* InNode = Cast<UNiagaraNodeReadDataSet>(Node))
+		{
+			OutReadNodes.Add(InNode);
+		}
+	}
+}
+
+void UNiagaraGraph::FindWriteDataSetNodes(TArray<class UNiagaraNodeWriteDataSet*>& OutWriteNodes) const
+{
+	for (UEdGraphNode* Node : Nodes)
+	{
+		if (UNiagaraNodeWriteDataSet* InNode = Cast<UNiagaraNodeWriteDataSet>(Node))
+		{
+			OutWriteNodes.Add(InNode);
+		}
+	}
+}
+
 int32 UNiagaraGraph::GetAttributeIndex(const FNiagaraVariableInfo& Attr)const
 {
 	const UNiagaraNodeOutput* OutNode = FindOutputNode();
@@ -99,10 +125,72 @@ void UNiagaraScriptSource::PostLoad()
 #endif
 }
 
+struct FNiagaraComponentReregisterContext : FComponentReregisterContext
+{
+	FNiagaraEmitterScriptProperties* ScriptProps;
+	UNiagaraEmitterProperties* EmitterProps;
+	FNiagaraComponentReregisterContext(UNiagaraComponent* Comp, FNiagaraEmitterScriptProperties* InScriptProps, UNiagaraEmitterProperties* InEmitterProps)
+		: FComponentReregisterContext(Comp)
+		, ScriptProps(InScriptProps)
+		, EmitterProps(InEmitterProps)
+	{
+
+	}
+	~FNiagaraComponentReregisterContext()
+	{
+		ScriptProps->Init(EmitterProps);
+	}
+};
+
+class FNiagaraScriptCompileContext
+{
+public:
+	/** Initialization constructor. */
+	FNiagaraScriptCompileContext(UNiagaraScript* Script)
+	{
+		// wait until resources are released
+		FlushRenderingCommands();
+
+		// Reregister all components usimg Script.
+		for (TObjectIterator<UNiagaraComponent> ComponentIt; ComponentIt; ++ComponentIt)
+		{
+			UNiagaraComponent* Comp = *ComponentIt;
+			TSharedPtr<FNiagaraEffectInstance> Inst = Comp->GetEffectInstance();
+			if (Inst.IsValid())
+			{
+				TArray<TSharedPtr<FNiagaraSimulation>>& Emitters = Inst->GetEmitters();
+				for (TSharedPtr<FNiagaraSimulation> Sim : Emitters)
+				{
+					if (Sim.IsValid())
+					{
+						if (UNiagaraEmitterProperties* Props = Sim->GetProperties().Get())
+						{
+							if (Props->UpdateScriptProps.Script == Script)
+							{
+								new(ComponentContexts)FNiagaraComponentReregisterContext(Comp, &Props->UpdateScriptProps, Props);
+							}
+							else if (Props->SpawnScriptProps.Script == Script)
+							{
+								new(ComponentContexts)FNiagaraComponentReregisterContext(Comp, &Props->SpawnScriptProps, Props);
+							}
+						}						
+					}
+				}
+			}
+		}
+	}
+	
+private:
+	/** The recreate contexts for the individual components. */
+	TIndirectArray<FNiagaraComponentReregisterContext> ComponentContexts;
+};
 
 void UNiagaraScriptSource::Compile()
 {
 	UNiagaraScript* ScriptOwner = Cast<UNiagaraScript>(GetOuter());
+	
+	FNiagaraScriptCompileContext ScriptCompileContext(ScriptOwner);
+
 	FNiagaraEditorModule& NiagaraEditorModule = FModuleManager::Get().LoadModuleChecked<FNiagaraEditorModule>(TEXT("NiagaraEditor"));
 	NiagaraEditorModule.CompileScript(ScriptOwner);
 
diff --git a/Engine/Source/Editor/UnrealEd/Public/INiagaraCompiler.h b/Engine/Source/Editor/UnrealEd/Public/INiagaraCompiler.h
--- a/Engine/Source/Editor/UnrealEd/Public/INiagaraCompiler.h
+++ b/Engine/Source/Editor/UnrealEd/Public/INiagaraCompiler.h
@@ -41,6 +41,18 @@ public:
 	/** Creates an expression which writes the result of SourceExpression to the named attribute. */
 	virtual TNiagaraExprPtr Output(const FNiagaraVariableInfo& Attr, TNiagaraExprPtr& SourceExpression) = 0;
 
+	/** Creates an expression which reads data from a shared data set. */
+	virtual TNiagaraExprPtr SharedDataRead(const FNiagaraDataSetID& DataSet, const FNiagaraVariableInfo& Variable, TNiagaraExprPtr& IndexExpr) = 0;
+	/** Creates an expression which writes data to an shared data set. */
+	virtual TNiagaraExprPtr SharedDataWrite(const FNiagaraDataSetID& DataSet, const FNiagaraVariableInfo& Variable, TNiagaraExprPtr& IndexExpr, TNiagaraExprPtr& DataExpr) = 0;
+
+	/** Creates an expression which gets the current counter for the shared data set and increments the counter. */
+	virtual TNiagaraExprPtr AcquireSharedDataIndex(const FNiagaraDataSetID& DataSet, bool bWrap, TNiagaraExprPtr& ValidExpr) = 0;
+	/** Creates an expression which gets the current counter for the shared data set and decrements the counter. */
+	virtual TNiagaraExprPtr ConsumeSharedDataIndex(const FNiagaraDataSetID& DataSet, bool bWrap) = 0;
+	/** Creates an expression with a mask of 0 or 0xFFFFFFFF marking the passed index as valid or not.*/
+	virtual TNiagaraExprPtr SharedDataIndexIsValid(const FNiagaraDataSetID& DataSet, TNiagaraExprPtr& IndexExpr) = 0;
+
 	/** Checks the validity of the input expressions against the informations about this operation in FNiagaraOpInfo. */
 	virtual bool CheckInputs(FName OpName, TArray<TNiagaraExprPtr>& Inputs) = 0;
 	/** Checks the validity of the output expressions against the informations about this operation in FNiagaraOpInfo. */
diff --git a/Engine/Source/Editor/UnrealEd/Public/NiagaraEditorCommon.h b/Engine/Source/Editor/UnrealEd/Public/NiagaraEditorCommon.h
--- a/Engine/Source/Editor/UnrealEd/Public/NiagaraEditorCommon.h
+++ b/Engine/Source/Editor/UnrealEd/Public/NiagaraEditorCommon.h
@@ -64,11 +64,12 @@
 	NiagaraOp(Transpose)\
 	NiagaraOp(Inverse)\
 	NiagaraOp(LessThan)\
+	NiagaraOp(GreaterThan)\
+	NiagaraOp(Select)\
 	NiagaraOp(Sample)\
 	NiagaraOp(Write)\
-	NiagaraOp(EventBroadcast)\
 	NiagaraOp(EaseIn)\
-	NiagaraOp(EaseInOut)\
+	NiagaraOp(EaseInOut)
 
 enum class ENiagaraExpressionResultLocation
 {
@@ -77,6 +78,7 @@ enum class ENiagaraExpressionResultLocation
 	Constants,
 	Temporaries,
 	BufferConstants,
+	SharedData,
 	Unknown,
 };
 
@@ -84,6 +86,13 @@ enum class ENiagaraExpressionResultLocation
 
 typedef TSharedPtr<class FNiagaraExpression> TNiagaraExprPtr;
 
+struct FNiagaraExpressionResult
+{
+	int32 Index;
+
+	ENiagaraExpressionResultLocation ResultLocation;
+};
+
 /** Base class for all Niagara Expresions. Intermediate data collected from the node graph used for compiling Niagara scripts. */
 class FNiagaraExpression : public TSharedFromThis<FNiagaraExpression>
 {
@@ -92,7 +101,6 @@ public:
 		: Compiler(InCompiler)
 		, ResultLocation(ENiagaraExpressionResultLocation::Unknown)
 		, ResultIndex(INDEX_NONE)
-		, ComponentIndex(INDEX_NONE)
 		, Result(InResult)
 	{
 	}
@@ -101,8 +109,8 @@ public:
 	{
 	}
 
-	/** Perform the work specific to the expression type. */
-	virtual void Process(){}
+	/** Perform the work specific to the expression type. Returns success or faliure. */
+	virtual bool Process(){ return true; }
 
 	virtual void PostProcess()
 	{
@@ -120,9 +128,6 @@ public:
 	/** Index of this result within it's location. */
 	int32 ResultIndex;
 
-	/** The index of the component we're accessing. Typically this is unused and set to INDEX_NONE. Scalar ops will set and use it however. */
-	int32 ComponentIndex;
-
 	/** The variable info for the result of this expression. */
 	FNiagaraVariableInfo Result;
 
diff --git a/Engine/Source/Runtime/Core/Public/Math/UnrealMathVectorConstants.h b/Engine/Source/Runtime/Core/Public/Math/UnrealMathVectorConstants.h
--- a/Engine/Source/Runtime/Core/Public/Math/UnrealMathVectorConstants.h
+++ b/Engine/Source/Runtime/Core/Public/Math/UnrealMathVectorConstants.h
@@ -18,6 +18,7 @@ namespace GlobalVectorConstants
 	static const VectorRegister KindaSmallNumber = MakeVectorRegister( KINDA_SMALL_NUMBER, KINDA_SMALL_NUMBER, KINDA_SMALL_NUMBER, KINDA_SMALL_NUMBER );
 	static const VectorRegister SmallNumber = MakeVectorRegister( SMALL_NUMBER, SMALL_NUMBER, SMALL_NUMBER, SMALL_NUMBER );
 	static const VectorRegister ThreshQuatNormalized = MakeVectorRegister( THRESH_QUAT_NORMALIZED, THRESH_QUAT_NORMALIZED, THRESH_QUAT_NORMALIZED, THRESH_QUAT_NORMALIZED );
+	static const VectorRegister BigNumber = MakeVectorRegister(BIG_NUMBER, BIG_NUMBER, BIG_NUMBER, BIG_NUMBER);
 
 	/** This is to speed up Quaternion Inverse. Static variable to keep sign of inverse **/
 	static const VectorRegister QINV_SIGN_MASK = MakeVectorRegister( -1.f, -1.f, -1.f, 1.f );
@@ -38,6 +39,7 @@ namespace GlobalVectorConstants
 	static const VectorRegister SignBit = MakeVectorRegister((uint32)SIGN_BIT, (uint32)SIGN_BIT, (uint32)SIGN_BIT, (uint32)SIGN_BIT);
 	static const VectorRegister SignMask = MakeVectorRegister((uint32)(~SIGN_BIT), (uint32)(~SIGN_BIT), (uint32)(~SIGN_BIT), (uint32)(~SIGN_BIT));
 #undef SIGN_BIT
+	static const VectorRegister AllMask = MakeVectorRegister(0xFFFFFFFF, 0xFFFFFFFF, 0xFFFFFFFF, 0xFFFFFFFF);
 
 	/** Vector full of positive infinity */
 	static const VectorRegister FloatInfinity = MakeVectorRegister((uint32)0x7F800000, (uint32)0x7F800000, (uint32)0x7F800000, (uint32)0x7F800000);
diff --git a/Engine/Source/Runtime/Engine/Classes/Engine/NiagaraEffectRenderer.h b/Engine/Source/Runtime/Engine/Classes/Engine/NiagaraEffectRenderer.h
--- a/Engine/Source/Runtime/Engine/Classes/Engine/NiagaraEffectRenderer.h
+++ b/Engine/Source/Runtime/Engine/Classes/Engine/NiagaraEffectRenderer.h
@@ -10,6 +10,7 @@ NiagaraEffectRenderer.h: Base class for Niagara render modules
 #include "ParticleVertexFactory.h"
 #include "ParticleBeamTrailVertexFactory.h"
 #include "NiagaraComponent.h"
+#include "NiagaraSimulation.h"
 #include "NiagaraEffectRendererProperties.h"
 #include "NiagaraSpriteRendererProperties.h"
 #include "NiagaraRibbonRendererProperties.h"
@@ -60,7 +61,7 @@ public:
 	virtual void SetDynamicData_RenderThread(FNiagaraDynamicDataBase* NewDynamicData) = 0;
 	virtual void CreateRenderThreadResources() = 0;
 	virtual void ReleaseRenderThreadResources() = 0;
-	virtual FNiagaraDynamicDataBase *GenerateVertexData(const FNiagaraEmitterParticleData &Data) = 0;
+	virtual FNiagaraDynamicDataBase *GenerateVertexData(const FNiagaraDataSet &Data) = 0;
 	virtual int GetDynamicDataSize() = 0;
 
 	virtual bool HasDynamicData() = 0;
@@ -149,7 +150,7 @@ public:
 	virtual void GetDynamicMeshElements(const TArray<const FSceneView*>& Views, const FSceneViewFamily& ViewFamily, uint32 VisibilityMap, FMeshElementCollector& Collector, const FNiagaraSceneProxy *SceneProxy) const override;
 	virtual bool SetMaterialUsage() override;
 	/** Update render data buffer from attributes */
-	FNiagaraDynamicDataBase *GenerateVertexData(const FNiagaraEmitterParticleData &Data) override;
+	FNiagaraDynamicDataBase *GenerateVertexData(const FNiagaraDataSet &Data) override;
 
 	virtual void SetDynamicData_RenderThread(FNiagaraDynamicDataBase* NewDynamicData) override;
 	int GetDynamicDataSize() override;
@@ -193,9 +194,9 @@ public:
 	virtual bool SetMaterialUsage() override;
 
 	/** Update render data buffer from attributes */
-	FNiagaraDynamicDataBase *GenerateVertexData(const FNiagaraEmitterParticleData &Data) override;
+	FNiagaraDynamicDataBase *GenerateVertexData(const FNiagaraDataSet &Data) override;
 
-	void AddRibbonVert(TArray<FParticleBeamTrailVertex>& RenderData, FVector ParticlePos, const FNiagaraEmitterParticleData &Data, FVector2D UV1,
+	void AddRibbonVert(TArray<FParticleBeamTrailVertex>& RenderData, FVector ParticlePos, const FNiagaraDataSet &Data, FVector2D UV1,
 		const FVector4 &Color, const FVector4 &Age, const FVector4 &Rotation)
 	{
 		FParticleBeamTrailVertex& NewVertex = *new(RenderData)FParticleBeamTrailVertex;
@@ -258,7 +259,7 @@ public:
 	virtual void GetDynamicMeshElements(const TArray<const FSceneView*>& Views, const FSceneViewFamily& ViewFamily, uint32 VisibilityMap, FMeshElementCollector& Collector, const FNiagaraSceneProxy *SceneProxy) const override;
 	virtual bool SetMaterialUsage() override;
 	/** Update render data buffer from attributes */
-	FNiagaraDynamicDataBase *GenerateVertexData(const FNiagaraEmitterParticleData &Data) override;
+	FNiagaraDynamicDataBase *GenerateVertexData(const FNiagaraDataSet &Data) override;
 
 	virtual void SetDynamicData_RenderThread(FNiagaraDynamicDataBase* NewDynamicData) override;
 	int GetDynamicDataSize() override;
diff --git a/Engine/Source/Runtime/Engine/Private/NiagaraEffectRenderer.cpp b/Engine/Source/Runtime/Engine/Private/NiagaraEffectRenderer.cpp
--- a/Engine/Source/Runtime/Engine/Private/NiagaraEffectRenderer.cpp
+++ b/Engine/Source/Runtime/Engine/Private/NiagaraEffectRenderer.cpp
@@ -243,7 +243,7 @@ bool NiagaraEffectRendererSprites::SetMaterialUsage()
 }
 
 /** Update render data buffer from attributes */
-FNiagaraDynamicDataBase *NiagaraEffectRendererSprites::GenerateVertexData(const FNiagaraEmitterParticleData &Data)
+FNiagaraDynamicDataBase *NiagaraEffectRendererSprites::GenerateVertexData(const FNiagaraDataSet &Data)
 {
 	SCOPE_CYCLE_COUNTER(STAT_NiagaraGenSpriteVertexData);
 
@@ -253,14 +253,14 @@ FNiagaraDynamicDataBase *NiagaraEffectRendererSprites::GenerateVertexData(const
 	FNiagaraDynamicDataSprites *DynamicData = new FNiagaraDynamicDataSprites;
 	TArray<FParticleSpriteVertex>& RenderData = DynamicData->VertexData;
 
-	RenderData.Reset(Data.GetNumParticles());
+	RenderData.Reset(Data.GetNumInstances());
 	CachedBounds.Init();
 
-	const FVector4 *PosPtr = Data.GetAttributeData(FNiagaraVariableInfo(FName(TEXT("Position")), ENiagaraDataType::Vector));
-	const FVector4 *VelPtr = Data.GetAttributeData(FNiagaraVariableInfo(FName(TEXT("Velocity")), ENiagaraDataType::Vector));
-	const FVector4 *ColPtr = Data.GetAttributeData(FNiagaraVariableInfo(FName(TEXT("Color")), ENiagaraDataType::Vector));
-	const FVector4 *AgePtr = Data.GetAttributeData(FNiagaraVariableInfo(FName(TEXT("Age")), ENiagaraDataType::Vector));
-	const FVector4 *RotPtr = Data.GetAttributeData(FNiagaraVariableInfo(FName(TEXT("Rotation")), ENiagaraDataType::Vector));
+	const FVector4 *PosPtr = Data.GetVariableData(FNiagaraVariableInfo(FName(TEXT("Position")), ENiagaraDataType::Vector));
+	const FVector4 *VelPtr = Data.GetVariableData(FNiagaraVariableInfo(FName(TEXT("Velocity")), ENiagaraDataType::Vector));
+	const FVector4 *ColPtr = Data.GetVariableData(FNiagaraVariableInfo(FName(TEXT("Color")), ENiagaraDataType::Vector));
+	const FVector4 *AgePtr = Data.GetVariableData(FNiagaraVariableInfo(FName(TEXT("Age")), ENiagaraDataType::Vector));
+	const FVector4 *RotPtr = Data.GetVariableData(FNiagaraVariableInfo(FName(TEXT("Rotation")), ENiagaraDataType::Vector));
 
 	//Bail if we don't have the required attributes to render this emitter.
 	if (!PosPtr || !ColPtr || !AgePtr || !RotPtr)
@@ -273,9 +273,9 @@ FNiagaraDynamicDataBase *NiagaraEffectRendererSprites::GenerateVertexData(const
 	{
 		NumSubImages = Properties->SubImageInfo.X*Properties->SubImageInfo.Y;
 	}
-	float ParticleId = 0.0f, IdInc = 1.0f / Data.GetNumParticles();
-	RenderData.AddUninitialized(Data.GetNumParticles());
-	for (uint32 ParticleIndex = 0; ParticleIndex < Data.GetNumParticles(); ParticleIndex++)
+	float ParticleId = 0.0f, IdInc = 1.0f / Data.GetNumInstances();
+	RenderData.AddUninitialized(Data.GetNumInstances());
+	for (uint32 ParticleIndex = 0; ParticleIndex < Data.GetNumInstances(); ParticleIndex++)
 	{
 		FParticleSpriteVertex& NewVertex = RenderData[ParticleIndex];
 		NewVertex.Position = PosPtr[ParticleIndex];
@@ -528,7 +528,7 @@ bool NiagaraEffectRendererRibbon::SetMaterialUsage()
 	return Material && Material->CheckMaterialUsage_Concurrent(MATUSAGE_BeamTrails);
 }
 
-FNiagaraDynamicDataBase *NiagaraEffectRendererRibbon::GenerateVertexData(const FNiagaraEmitterParticleData &Data)
+FNiagaraDynamicDataBase *NiagaraEffectRendererRibbon::GenerateVertexData(const FNiagaraDataSet &Data)
 {
 	SCOPE_CYCLE_COUNTER(STAT_NiagaraGenRibbonVertexData);
 
@@ -537,21 +537,21 @@ FNiagaraDynamicDataBase *NiagaraEffectRendererRibbon::GenerateVertexData(const F
 	FNiagaraDynamicDataRibbon *DynamicData = new FNiagaraDynamicDataRibbon;
 	TArray<FParticleBeamTrailVertex>& RenderData = DynamicData->VertexData;
 
-	RenderData.Reset(Data.GetNumParticles() * 2);
+	RenderData.Reset(Data.GetNumInstances() * 2);
 	//CachedBounds.Init();
 
 	// build a sorted list by age, so we always get particles in order 
 	// regardless of them being moved around due to dieing and spawning
 	TArray<int32> SortedIndices;
-	for (uint32 Idx = 0; Idx < Data.GetNumParticles(); Idx++)
+	for (uint32 Idx = 0; Idx < Data.GetNumInstances(); Idx++)
 	{
 		SortedIndices.Add(Idx);
 	}
 
-	const FVector4 *PosPtr = Data.GetAttributeData(FNiagaraVariableInfo(FName(TEXT("Position")), ENiagaraDataType::Vector));
-	const FVector4 *ColorPtr = Data.GetAttributeData(FNiagaraVariableInfo(FName(TEXT("Color")), ENiagaraDataType::Vector));
-	const FVector4 *RotPtr = Data.GetAttributeData(FNiagaraVariableInfo(FName(TEXT("Rotation")), ENiagaraDataType::Vector));
-	const FVector4 *AgePtr = Data.GetAttributeData(FNiagaraVariableInfo(FName(TEXT("Age")), ENiagaraDataType::Vector));
+	const FVector4 *PosPtr = Data.GetVariableData(FNiagaraVariableInfo(FName(TEXT("Position")), ENiagaraDataType::Vector));
+	const FVector4 *ColorPtr = Data.GetVariableData(FNiagaraVariableInfo(FName(TEXT("Color")), ENiagaraDataType::Vector));
+	const FVector4 *RotPtr = Data.GetVariableData(FNiagaraVariableInfo(FName(TEXT("Rotation")), ENiagaraDataType::Vector));
+	const FVector4 *AgePtr = Data.GetVariableData(FNiagaraVariableInfo(FName(TEXT("Age")), ENiagaraDataType::Vector));
 
 	//Bail if we don't have the required attributes to render this emitter.
 	if (!PosPtr || !ColorPtr || !AgePtr || !RotPtr)
@@ -947,7 +947,7 @@ bool NiagaraEffectRendererMeshes::SetMaterialUsage()
 
 
 /** Update render data buffer from attributes */
-FNiagaraDynamicDataBase *NiagaraEffectRendererMeshes::GenerateVertexData(const FNiagaraEmitterParticleData &Data)
+FNiagaraDynamicDataBase *NiagaraEffectRendererMeshes::GenerateVertexData(const FNiagaraDataSet &Data)
 {
 	SCOPE_CYCLE_COUNTER(STAT_NiagaraGenMeshVertexData);
 
@@ -963,15 +963,15 @@ FNiagaraDynamicDataBase *NiagaraEffectRendererMeshes::GenerateVertexData(const F
 	FNiagaraDynamicDataMesh *DynamicData = new FNiagaraDynamicDataMesh;
 	TArray<FMeshParticleInstanceVertex>& RenderData = DynamicData->VertexData;
 
-	RenderData.Reset(Data.GetNumParticles());
+	RenderData.Reset(Data.GetNumInstances());
 	CachedBounds.Init();
 
-	const FVector4 *PosPtr = Data.GetAttributeData(FNiagaraVariableInfo(FName(TEXT("Position")), ENiagaraDataType::Vector));
-	const FVector4 *VelPtr = Data.GetAttributeData(FNiagaraVariableInfo(FName(TEXT("Velocity")), ENiagaraDataType::Vector));
-	const FVector4 *ColPtr = Data.GetAttributeData(FNiagaraVariableInfo(FName(TEXT("Color")), ENiagaraDataType::Vector));
-	const FVector4 *AgePtr = Data.GetAttributeData(FNiagaraVariableInfo(FName(TEXT("Age")), ENiagaraDataType::Vector));
-	const FVector4 *RotPtr = Data.GetAttributeData(FNiagaraVariableInfo(FName(TEXT("Rotation")), ENiagaraDataType::Vector));
-	const FVector4 *XformPtr = Data.GetAttributeData(FNiagaraVariableInfo(FName(TEXT("Transform")), ENiagaraDataType::Vector));
+	const FVector4 *PosPtr = Data.GetVariableData(FNiagaraVariableInfo(FName(TEXT("Position")), ENiagaraDataType::Vector));
+	const FVector4 *VelPtr = Data.GetVariableData(FNiagaraVariableInfo(FName(TEXT("Velocity")), ENiagaraDataType::Vector));
+	const FVector4 *ColPtr = Data.GetVariableData(FNiagaraVariableInfo(FName(TEXT("Color")), ENiagaraDataType::Vector));
+	const FVector4 *AgePtr = Data.GetVariableData(FNiagaraVariableInfo(FName(TEXT("Age")), ENiagaraDataType::Vector));
+	const FVector4 *RotPtr = Data.GetVariableData(FNiagaraVariableInfo(FName(TEXT("Rotation")), ENiagaraDataType::Vector));
+	const FVector4 *XformPtr = Data.GetVariableData(FNiagaraVariableInfo(FName(TEXT("Transform")), ENiagaraDataType::Vector));
 
 	//Bail if we don't have the required attributes to render this emitter.
 	if (!PosPtr || !ColPtr || !AgePtr || !RotPtr)
@@ -984,9 +984,9 @@ FNiagaraDynamicDataBase *NiagaraEffectRendererMeshes::GenerateVertexData(const F
 	{
 		Mesh = Properties->ParticleMesh;
 	}
-	float ParticleId = 0.0f, IdInc = 1.0f / Data.GetNumParticles();
-	RenderData.AddUninitialized(Data.GetNumParticles());
-	for (uint32 ParticleIndex = 0; ParticleIndex < Data.GetNumParticles(); ParticleIndex++)
+	float ParticleId = 0.0f, IdInc = 1.0f / Data.GetNumInstances();
+	RenderData.AddUninitialized(Data.GetNumInstances());
+	for (uint32 ParticleIndex = 0; ParticleIndex < Data.GetNumInstances(); ParticleIndex++)
 	{
 		FMeshParticleInstanceVertex& NewVertex = RenderData[ParticleIndex];
 		NewVertex.Color = FLinearColor(ColPtr[ParticleIndex]);
diff --git a/Engine/Source/Runtime/Niagara/Classes/NiagaraDataSet.h b/Engine/Source/Runtime/Niagara/Classes/NiagaraDataSet.h
new file mode 100644
--- /dev/null
+++ b/Engine/Source/Runtime/Niagara/Classes/NiagaraDataSet.h
@@ -0,0 +1,168 @@
+// Copyright 1998-2015 Epic Games, Inc. All Rights Reserved.
+#pragma once
+
+#include "NiagaraCommon.h"
+
+/** Storage for Niagara data. */
+class FNiagaraDataSet
+{
+public:
+	FNiagaraDataSet()
+	{
+		Reset();
+	}
+
+	FNiagaraDataSet(FNiagaraDataSetID InID)
+		: ID(InID)
+	{
+		Reset();
+	}
+
+	~FNiagaraDataSet()
+	{
+		Reset();
+	}
+
+	void Allocate(uint32 NumExpectedInstances, bool bReset = true)
+	{
+		if (bReset)
+		{
+			DataBuffers[CurrentBuffer].Reset(NumExpectedInstances * VariableMap.Num());
+			DataBuffers[CurrentBuffer].AddUninitialized(NumExpectedInstances * VariableMap.Num());
+		}
+		else
+		{
+			DataBuffers[CurrentBuffer].SetNumUninitialized(NumExpectedInstances * VariableMap.Num());
+		}
+
+		DataAllocation[CurrentBuffer] = NumExpectedInstances;
+	}
+
+	void Reset()
+	{
+		VariableMap.Empty();
+		NumInstances[0] = 0;
+		NumInstances[1] = 0;
+		DataAllocation[0] = 0;
+		DataAllocation[1] = 0;
+		CurrentBuffer = 0;
+		DataBuffers[0].Empty();
+		DataBuffers[1].Empty();
+	}
+
+	FVector4* GetVariableData(const FNiagaraVariableInfo& VariableID, uint32 StartParticle = 0)
+	{
+		const uint32* Offset = VariableMap.Find(VariableID);
+		return Offset ? DataBuffers[CurrentBuffer].GetData() + (*Offset * DataAllocation[CurrentBuffer]) + StartParticle : NULL;
+	}
+
+	const FVector4* GetVariableData(const FNiagaraVariableInfo& VariableID, uint32 StartParticle = 0)const
+	{
+		const uint32* Offset = VariableMap.Find(VariableID);
+		return Offset ? DataBuffers[CurrentBuffer].GetData() + (*Offset * DataAllocation[CurrentBuffer]) + StartParticle : NULL;
+	}
+
+	FVector4* GetPrevVariableData(const FNiagaraVariableInfo& VariableID, uint32 StartParticle = 0)
+	{
+		const uint32* Offset = VariableMap.Find(VariableID);
+		return Offset ? DataBuffers[CurrentBuffer ^ 0x1].GetData() + (*Offset * DataAllocation[CurrentBuffer ^ 0x1]) + StartParticle : NULL;
+	}
+
+	const FVector4* GetPrevVariableData(const FNiagaraVariableInfo& VariableID, uint32 StartParticle = 0)const
+	{
+		const uint32* Offset = VariableMap.Find(VariableID);
+		return Offset ? DataBuffers[CurrentBuffer ^ 0x1].GetData() + (*Offset * DataAllocation[CurrentBuffer ^ 0x1]) + StartParticle : NULL;
+	}
+
+	void GetVariableData(const FNiagaraVariableInfo& VariableID, FVector4*& PrevBuffer, FVector4*& CurrBuffer, uint32 StartParticle)
+	{
+		const uint32* Offset = VariableMap.Find(VariableID);
+		if (Offset)
+		{
+			PrevBuffer = DataBuffers[CurrentBuffer ^ 0x1].GetData() + (*Offset * DataAllocation[CurrentBuffer ^ 0x1]) + StartParticle;
+			CurrBuffer = DataBuffers[CurrentBuffer].GetData() + (*Offset * DataAllocation[CurrentBuffer]) + StartParticle;
+		}
+		else
+		{
+			PrevBuffer = NULL;
+			CurrBuffer = NULL;
+		}
+	}
+
+	bool HasAttriubte(const FNiagaraVariableInfo& VariableID)
+	{
+		return VariableMap.Find(VariableID) != NULL;
+	}
+
+	void SetVariables(const TArray<FNiagaraVariableInfo>& VariableIDs)
+	{
+		Reset();
+		for (const FNiagaraVariableInfo& Var : VariableIDs)
+		{
+			//TODO: Support attributes of any type. This may as well wait until the move of matrix ops etc over into UNiagaraScriptStructs.
+			uint32 Idx = VariableMap.Num();
+			VariableMap.Add(Var, Idx);
+		}
+	}
+
+	void AddVariables(const TArray<FNiagaraVariableInfo>& Variables)
+	{
+		for (const FNiagaraVariableInfo& Variable : Variables)
+		{
+			AddVariable(Variable);
+		}
+	}
+
+	void AddVariable(const FNiagaraVariableInfo& Variable)
+	{
+		//TODO: Support attributes of any type. This may as well wait until the move of matrix ops etc over into UNiagaraScriptStructs.
+		uint32 Idx = VariableMap.Num();
+		if (!VariableMap.Find(Variable))
+		{
+			VariableMap.Add(Variable, Idx);
+		}
+	}
+
+	const TMap<FNiagaraVariableInfo, uint32>& GetVariables()const { return VariableMap; }
+	int GetNumVariables()				{ return VariableMap.Num(); }
+
+	uint32 GetNumInstances() const		{ return NumInstances[CurrentBuffer]; }
+	uint32 GetPrevNumInstances() const		{ return NumInstances[CurrentBuffer ^ 0x1]; }
+	uint32 GetDataAllocation() const { return DataAllocation[CurrentBuffer]; }
+	uint32 GetPrevDataAllocation() const { return DataAllocation[CurrentBuffer ^ 0x1]; }
+	void SetNumInstances(uint32 Num)	{ NumInstances[CurrentBuffer] = Num; }
+	void SwapBuffers()					{ CurrentBuffer ^= 0x1; }
+
+	FVector4* GetCurrentBuffer()		{ return DataBuffers[CurrentBuffer].GetData(); }
+	FVector4* GetPreviousBuffer()		{ return DataBuffers[CurrentBuffer ^ 0x1].GetData(); }
+
+	int GetBytesUsed()	{ return (DataBuffers[0].Num() + DataBuffers[1].Num()) * 16 + VariableMap.Num() * 4; }
+
+	FNiagaraDataSetID GetID()const { return ID; }
+	void SetID(FNiagaraDataSetID InID) { ID = InID; }
+
+	void Tick()
+	{
+		SwapBuffers();
+
+		if (ID.Type == ENiagaraDataSetType::Event)
+		{
+			SetNumInstances(0);
+			Allocate(GetPrevDataAllocation());
+			//Fancier allocation for events?
+		}
+		else if (ID.Type == ENiagaraDataSetType::ParticleData)
+		{
+			SetNumInstances(GetPrevNumInstances());
+		}
+	}
+
+private:
+	uint32 CurrentBuffer;
+	uint32 NumInstances[2];
+	uint32 DataAllocation[2];
+	TArray<FVector4> DataBuffers[2];
+	TMap<FNiagaraVariableInfo, uint32> VariableMap;
+	FNiagaraDataSetID ID;
+};
+
diff --git a/Engine/Source/Runtime/Niagara/Classes/NiagaraEffect.h b/Engine/Source/Runtime/Niagara/Classes/NiagaraEffect.h
--- a/Engine/Source/Runtime/Niagara/Classes/NiagaraEffect.h
+++ b/Engine/Source/Runtime/Niagara/Classes/NiagaraEffect.h
@@ -2,6 +2,7 @@
 
 #pragma once
 #include "NiagaraSimulation.h"
+#include "NiagaraEmitterProperties.h"
 #include "Engine/NiagaraEffectRenderer.h"
 #include "List.h"
 #include "Runtime/VectorVM/Public/VectorVMDataObject.h"
@@ -44,7 +45,7 @@ public:
 	void CreateEffectRendererProps(TSharedPtr<FNiagaraSimulation> Sim);
 
 	//Begin UObject Interface
-	virtual void Serialize(FArchive& Ar);
+	virtual void PostLoad()override;
 	//End UObject Interface
 private:
 	UPROPERTY()
@@ -65,7 +66,7 @@ public:
 		Component = InComponent;
 		for (int32 i = 0; i < InAsset->GetNumEmitters(); i++)
 		{
-			FNiagaraSimulation *Sim = new FNiagaraSimulation(InAsset->GetEmitterProperties(i), InAsset, Component->GetWorld()->FeatureLevel);
+			FNiagaraSimulation *Sim = new FNiagaraSimulation(InAsset->GetEmitterProperties(i), this, Component->GetWorld()->FeatureLevel);
 			Emitters.Add(MakeShareable(Sim));
 		}
 		InitRenderModules(Component->GetWorld()->FeatureLevel);
@@ -81,6 +82,7 @@ public:
 	}
 
 	NIAGARA_API void InitEmitters(UNiagaraEffect *InAsset);
+	NIAGARA_API void ReInitEmitters();
 	void InitRenderModules(ERHIFeatureLevel::Type InFeatureLevel)
 	{
 		for (TSharedPtr<FNiagaraSimulation> Sim : Emitters)
@@ -107,6 +109,17 @@ public:
 		Age = 0.0f;
 	}
 
+	void ReInit()
+	{
+		if (Component)
+		{
+			ReInitEmitters();
+			InitRenderModules(Component->GetWorld()->FeatureLevel);
+			RenderModuleupdate();
+		}
+		Age = 0.0f;
+	}
+
 	NIAGARA_API TSharedPtr<FNiagaraSimulation> AddEmitter(UNiagaraEmitterProperties *Properties);
 	NIAGARA_API void DeleteEmitter(TSharedPtr<FNiagaraSimulation> Emitter);
 
@@ -139,12 +152,17 @@ public:
 		return Emitters[idx].Get();
 	}
 
+	FNiagaraSimulation* GetEmitter(FString Name);
 
 	UNiagaraComponent *GetComponent() { return Component; }
 
 	TArray< TSharedPtr<FNiagaraSimulation> > &GetEmitters()	{ return Emitters; }
 
 	FBox GetEffectBounds()	{ return EffectBounds;  }
+
+	/** Gets a data set either from another emitter or one owned by the effect itself. */
+	FNiagaraDataSet* GetDataSet(FNiagaraDataSetID SetID, FName EmitterName=NAME_None);
+	
 private:
 	UNiagaraComponent *Component;
 	UNiagaraEffect *Effect;
@@ -153,6 +171,8 @@ private:
 
 	/** Local constant table. */
 	FNiagaraConstantMap Constants;
+	
+	TMap<FNiagaraDataSetID, FNiagaraDataSet> ExternalEvents;
 
 	TArray< TSharedPtr<FNiagaraSimulation> > Emitters;
 
diff --git a/Engine/Source/Runtime/Niagara/Classes/NiagaraEmitterProperties.h b/Engine/Source/Runtime/Niagara/Classes/NiagaraEmitterProperties.h
new file mode 100644
--- /dev/null
+++ b/Engine/Source/Runtime/Niagara/Classes/NiagaraEmitterProperties.h
@@ -0,0 +1,187 @@
+// Copyright 1998-2015 Epic Games, Inc. All Rights Reserved.
+
+#pragma once
+
+#include "NiagaraScript.h"
+#include "NiagaraDataSet.h"
+#include "NiagaraEmitterProperties.generated.h"
+
+class NiagaraEffectRenderer;
+class UNiagaraEventReceiverEmitterAction;
+
+UENUM()
+enum EEmitterRenderModuleType
+{
+	RMT_None = 0,
+	RMT_Sprites,
+	RMT_Ribbon,
+	RMT_Trails,
+	RMT_Meshes
+};
+
+//TODO: Event action that spawns other whole effects?
+//One that calls a BP exposed delegate?
+
+USTRUCT()
+struct FNiagaraEventReceiverProperties
+{
+	GENERATED_BODY()
+
+	FNiagaraEventReceiverProperties()
+	: Name(NAME_None)
+	, SourceEventGenerator(NAME_None)
+	, SourceEmitter(NAME_None)
+	{
+
+	}
+
+	FNiagaraEventReceiverProperties(FName InName, FName InEventGenerator, FName InSourceEmitter)
+		: Name(InName)
+		, SourceEventGenerator(InEventGenerator)
+		, SourceEmitter(InSourceEmitter)
+	{
+
+	}
+
+	/** The name of this receiver. */
+	UPROPERTY(EditAnywhere, Category = "Event Receiver")
+	FName Name;
+
+	/** The name of the EventGenerator to bind to. */
+	UPROPERTY(EditAnywhere, Category = "Event Receiver")
+	FName SourceEventGenerator;
+
+	/** The name of the emitter from which the Event Generator is taken. */
+	UPROPERTY(EditAnywhere, Category = "Event Receiver")
+	FName SourceEmitter;
+
+	UPROPERTY(EditAnywhere, Instanced, Category = "Event Receiver")
+	TArray<UNiagaraEventReceiverEmitterAction*> EmitterActions;
+};
+
+USTRUCT()
+struct FNiagaraEventGeneratorProperties : public FNiagaraDataSetProperties
+{
+	GENERATED_BODY()
+
+	FNiagaraEventGeneratorProperties()
+	: MaxEventsPerFrame(64)
+	{
+
+	}
+
+	FNiagaraEventGeneratorProperties(FNiagaraDataSetProperties& Props, int32 InMaxEventsPerFrame)
+	: FNiagaraDataSetProperties(Props)
+	, MaxEventsPerFrame(InMaxEventsPerFrame)
+	{
+
+	}
+
+	/** Max Number of Events that can be generated per frame. */
+	UPROPERTY(EditAnywhere, Category = "Event Receiver")
+	int32 MaxEventsPerFrame; //TODO - More complex allocation so that we can grow dynamically if more space is needed ?
+};
+
+USTRUCT()
+struct FNiagaraEmitterScriptProperties
+{
+	GENERATED_BODY()
+	
+ 	UPROPERTY(EditAnywhere, Category = "Script")
+ 	UNiagaraScript *Script;
+
+	UPROPERTY(EditAnywhere, Category = "Script")
+ 	FNiagaraConstants ExternalConstants;
+
+	UPROPERTY(EditAnywhere, EditFixedSize, Category = "Script")
+	TArray<FNiagaraEventReceiverProperties> EventReceivers;
+
+	UPROPERTY(EditAnywhere, EditFixedSize, Category = "Script")
+	TArray<FNiagaraEventGeneratorProperties> EventGenerators;
+
+	NIAGARA_API void Init(UNiagaraEmitterProperties* EmitterProps);
+};
+
+//This struct now only exists for backwards compatibility and should be removed once effects are updated.
+USTRUCT()
+struct FDeprecatedNiagaraEmitterProperties
+{
+	GENERATED_USTRUCT_BODY()
+public:
+	UPROPERTY()
+	FString Name;
+	UPROPERTY()
+	bool bIsEnabled;
+	UPROPERTY()
+	float SpawnRate;
+	UPROPERTY()
+	UNiagaraScript *UpdateScript;
+	UPROPERTY()
+	UNiagaraScript *SpawnScript;
+	UPROPERTY()
+	UMaterial *Material;
+	UPROPERTY()
+	TEnumAsByte<EEmitterRenderModuleType> RenderModuleType;
+	UPROPERTY()
+	float StartTime;
+	UPROPERTY()
+	float EndTime;
+	UPROPERTY()
+	class UNiagaraEffectRendererProperties *RendererProperties;
+	UPROPERTY()
+	FNiagaraConstantMap ExternalConstants;		// these are the update script constants from the effect editor; will be added to the emitter's constant map
+	UPROPERTY()
+	FNiagaraConstantMap ExternalSpawnConstants;		// these are the spawn script constants from the effect editor; will be added to the emitter's constant map
+	UPROPERTY()
+	int32 NumLoops;
+};
+
+/** 
+ *	UNiagaraEmitterProperties stores the attributes of an FNiagaraSimulation
+ *	that need to be serialized and are used for its initialization 
+ */
+UCLASS(MinimalAPI)
+class UNiagaraEmitterProperties : public UObject
+{
+	GENERATED_UCLASS_BODY()
+public:
+	NIAGARA_API void Init();
+	void InitFromOldStruct(FDeprecatedNiagaraEmitterProperties& OldStruct);
+
+	//Begin UObject Interface
+#if WITH_EDITOR
+	virtual void PostEditChangeProperty(struct FPropertyChangedEvent& PropertyChangedEvent) override;
+#endif
+	//End UObject Interface
+
+	UPROPERTY(EditAnywhere, Category = "Emitter")
+	FString EmitterName;
+	UPROPERTY(EditAnywhere, Category = "Emitter")
+	bool bIsEnabled;
+	UPROPERTY(EditAnywhere, Category = "Emitter")
+	float SpawnRate;
+	UPROPERTY(EditAnywhere, Category = "Emitter")
+	UMaterial *Material;
+	UPROPERTY(EditAnywhere, Category = "Emitter")
+	TEnumAsByte<EEmitterRenderModuleType> RenderModuleType;
+
+	UPROPERTY(EditAnywhere, Category = "Emitter")
+	float StartTime;
+	UPROPERTY(EditAnywhere, Category = "Emitter")
+	float EndTime;
+
+	//Can get rid of the enum and just have users select a class for this directly in the UI?
+	UPROPERTY(VisibleAnywhere, Instanced, Category = "Emitter")
+	class UNiagaraEffectRendererProperties *RendererProperties;
+
+	UPROPERTY(EditAnywhere, Category = "Emitter")
+	int32 NumLoops;
+
+	UPROPERTY(EditAnywhere, Category = "Emitter", meta = (ShowOnlyInnerProperties))
+	FNiagaraEmitterScriptProperties UpdateScriptProps;
+
+	UPROPERTY(EditAnywhere, Category = "Emitter", meta = (ShowOnlyInnerProperties))
+	FNiagaraEmitterScriptProperties SpawnScriptProps;
+};
+
+
diff --git a/Engine/Source/Runtime/Niagara/Classes/NiagaraEvents.h b/Engine/Source/Runtime/Niagara/Classes/NiagaraEvents.h
new file mode 100644
--- /dev/null
+++ b/Engine/Source/Runtime/Niagara/Classes/NiagaraEvents.h
@@ -0,0 +1,114 @@
+// Copyright 1998-2015 Epic Games, Inc. All Rights Reserved.
+
+/*==============================================================================
+NiagaraSimulation.h: Niagara emitter simulation class
+==============================================================================*/
+
+#pragma once
+
+#include "NiagaraDataSet.h"
+#include "NiagaraEvents.generated.h"
+
+struct FNiagaraSimulation;
+struct FNiagaraEventReceiverProperties;
+
+/**
+Base class for actions that an event receiver will perform at the emitter level.
+*/
+UCLASS(abstract)
+class UNiagaraEventReceiverEmitterAction : public UObject
+{
+	GENERATED_BODY()
+public:
+	virtual void PerformAction(FNiagaraSimulation& OwningSim, FNiagaraEventReceiverProperties& OwningEventReceiver){}
+};
+
+UCLASS(EditInlineNew)
+class UNiagaraEventReceiverEmitterAction_SpawnParticles : public UNiagaraEventReceiverEmitterAction
+{
+	GENERATED_BODY()
+
+public:
+
+	/** Number of particles to spawn per event received. */
+	UPROPERTY(EditAnywhere, Category = "Spawn")
+	uint32 NumParticles;
+
+	virtual void PerformAction(FNiagaraSimulation& OwningSim, FNiagaraEventReceiverProperties& OwningEventReceiver)override;
+};
+
+/*
+* Generates an event for every particle death in an emitter. Each event contains all particle attributes at the moment of death.
+*/
+class FNiagaraDeathEventGenerator
+{
+	/** The owning simulation. */
+	FNiagaraSimulation* OwnerSim;
+
+	/** Storage for dead particles. Source of event data. */
+	FNiagaraDataSet DataSet;
+
+	/*
+	Temp storage for writing the attribute data of dead particles.
+	Can likely get rid of this and write directly to the DataSet with a bit more work.
+	*/
+	TArray<TArray<FVector4>> TempStorage;
+
+	/** The number of dead particles we've seen this frame. */
+	int32 NumDead;
+public:
+
+	FNiagaraDeathEventGenerator(FNiagaraSimulation* InOwnerSim)
+		: OwnerSim(InOwnerSim)
+		, DataSet(FNiagaraDataSetID::DeathEvent)
+		, NumDead(0)
+	{
+
+	}
+
+	void Reset()
+	{
+		DataSet.Reset();
+		NumDead = 0;
+		TempStorage.Empty();
+	}
+
+	FNiagaraDataSet& GetDataSet(){ return DataSet; }
+	void BeginTrackingDeaths();
+	void OnDeath(int32 DeadIndex);
+	void EndTrackingDeaths();
+};
+
+/**
+* Generates events for every particle spawned in an emitter. Each event contains the particles attriubtes as initialised by the spawn script.
+*/
+class FNiagaraSpawnEventGenerator
+{
+	FNiagaraSimulation* OwnerSim;
+
+	/**
+	Storage for spawned particles. Source of event data.
+	For quickness, we currently make a copy of all spawned particles.
+	It may be possible in future to remove this and access the Particle data directly with some refactoring of the events stuff.
+	*/
+	FNiagaraDataSet DataSet;
+public:
+
+	void Reset()
+	{
+		DataSet.Reset();
+	}
+
+	FNiagaraDataSet& GetDataSet(){ return DataSet; }
+
+	void OnSpawned(int32 SpawnIndex, int32 NumSpawned);
+
+	FNiagaraSpawnEventGenerator(FNiagaraSimulation* InOwnerSim)
+		: OwnerSim(InOwnerSim)
+		, DataSet(FNiagaraDataSetID::SpawnEvent)
+	{
+
+	}
+};
+
+
diff --git a/Engine/Source/Runtime/Niagara/Classes/NiagaraScript.h b/Engine/Source/Runtime/Niagara/Classes/NiagaraScript.h
--- a/Engine/Source/Runtime/Niagara/Classes/NiagaraScript.h
+++ b/Engine/Source/Runtime/Niagara/Classes/NiagaraScript.h
@@ -22,6 +22,33 @@ enum class EUnusedAttributeBehaviour : uint8
 	MarkInvalid, 
 };
 
+USTRUCT()
+struct FNiagaraDataSetProperties
+{
+	GENERATED_BODY()
+	
+	UPROPERTY(VisibleAnywhere, Category = "Data Set")
+	FNiagaraDataSetID ID;
+
+	UPROPERTY()
+	TArray<FNiagaraVariableInfo> Variables;
+};
+
+/** Struct containt usage information about a script. Things such as whether it reads attribute data, reads or writes events data etc.*/
+USTRUCT()
+struct FNiagaraScriptUsageInfo
+{
+	GENERATED_BODY()
+	
+	FNiagaraScriptUsageInfo()
+	: bReadsAttriubteData(false)
+	{}
+
+	/** If true, this script reads attribute data. */
+	UPROPERTY()
+	bool bReadsAttriubteData;
+};
+
 /** Runtime script for a Niagara system */
 UCLASS(MinimalAPI)
 class UNiagaraScript : public UObject
@@ -40,6 +67,18 @@ class UNiagaraScript : public UObject
 	UPROPERTY()
  	TArray<FNiagaraVariableInfo> Attributes;
 
+	/** Information about the events this script receives and which variables are accessed. */
+	UPROPERTY()
+	TArray<FNiagaraDataSetProperties> EventReceivers;
+
+	/** Information about the events this script generates and which variables are written. */
+	UPROPERTY()
+	TArray<FNiagaraDataSetProperties> EventGenerators;
+
+	/** Contains various usage information for this script. */
+	UPROPERTY()
+	FNiagaraScriptUsageInfo Usage;
+
 #if WITH_EDITORONLY_DATA
 	/** 'Source' data/graphs for this script */
 	UPROPERTY()
