Bug fix for constant editing in custom UI not distinguishing properly between spawn and update constants; re-setup some of the effects having trouble due to this.

diff --git a/Engine/Source/Editor/NiagaraEditor/Private/SNiagaraEffectEditorWidget.cpp b/Engine/Source/Editor/NiagaraEditor/Private/SNiagaraEffectEditorWidget.cpp
--- a/Engine/Source/Editor/NiagaraEditor/Private/SNiagaraEffectEditorWidget.cpp
+++ b/Engine/Source/Editor/NiagaraEditor/Private/SNiagaraEffectEditorWidget.cpp
@@ -415,7 +415,7 @@ void SEmitterWidget::Construct(const FArguments& InArgs)
 										.BodyContent()
 										[
 											SAssignNew(UpdateScriptConstantList, SListView<TSharedPtr<EditorExposedVectorConstant>>)
-											.ItemHeight(20).ListItemsSource(EditorExposedConstants).OnGenerateRow(this, &SEmitterWidget::OnGenerateConstantListRow)
+											.ItemHeight(20).ListItemsSource(EditorExposedConstants).OnGenerateRow(this, &SEmitterWidget::OnGenerateUpdateConstantListRow)
 											.SelectionMode(ESelectionMode::None)
 										]
 									]
@@ -429,7 +429,7 @@ void SEmitterWidget::Construct(const FArguments& InArgs)
 											.BodyContent()
 											[
 												SNew(SListView<TSharedPtr<EditorExposedVectorCurveConstant>>)
-												.ItemHeight(20).ListItemsSource(EditorExposedCurveConstants).OnGenerateRow(this, &SEmitterWidget::OnGenerateCurveConstantListRow)
+												.ItemHeight(20).ListItemsSource(EditorExposedCurveConstants).OnGenerateRow(this, &SEmitterWidget::OnGenerateUpdateCurveConstantListRow)
 												.SelectionMode(ESelectionMode::None)
 											]
 										]
@@ -479,7 +479,7 @@ void SEmitterWidget::Construct(const FArguments& InArgs)
 											.BodyContent()
 											[
 												SAssignNew(SpawnScriptConstantList, SListView<TSharedPtr<EditorExposedVectorConstant>>)
-												.ItemHeight(20).ListItemsSource(EditorExposedSpawnConstants).OnGenerateRow(this, &SEmitterWidget::OnGenerateConstantListRow)
+												.ItemHeight(20).ListItemsSource(EditorExposedSpawnConstants).OnGenerateRow(this, &SEmitterWidget::OnGenerateSpawnConstantListRow)
 												.SelectionMode(ESelectionMode::None)
 											]
 										]
@@ -706,7 +706,7 @@ void SEmitterWidget::OnUpdateScriptSelectedFromPicker(UObject *Asset)
 				.BodyContent()
 				[
 					SAssignNew(UpdateScriptConstantList, SListView<TSharedPtr<EditorExposedVectorConstant>>)
-					.ItemHeight(20).ListItemsSource(&CurUpdateScript->Source->ExposedVectorConstants).OnGenerateRow(this, &SEmitterWidget::OnGenerateConstantListRow)
+					.ItemHeight(20).ListItemsSource(&CurUpdateScript->Source->ExposedVectorConstants).OnGenerateRow(this, &SEmitterWidget::OnGenerateUpdateConstantListRow)
 				]
 			];
 	}
@@ -734,7 +734,7 @@ void SEmitterWidget::OnSpawnScriptSelectedFromPicker(UObject *Asset)
 				.BodyContent()
 				[
 					SAssignNew(SpawnScriptConstantList, SListView<TSharedPtr<EditorExposedVectorConstant>>)
-					.ItemHeight(20).ListItemsSource(&CurSpawnScript->Source->ExposedVectorConstants).OnGenerateRow(this, &SEmitterWidget::OnGenerateConstantListRow)
+					.ItemHeight(20).ListItemsSource(&CurSpawnScript->Source->ExposedVectorConstants).OnGenerateRow(this, &SEmitterWidget::OnGenerateSpawnConstantListRow)
 				]
 			];
 	}
diff --git a/Engine/Source/Editor/NiagaraEditor/Private/SNiagaraEffectEditorWidget.h b/Engine/Source/Editor/NiagaraEditor/Private/SNiagaraEffectEditorWidget.h
--- a/Engine/Source/Editor/NiagaraEditor/Private/SNiagaraEffectEditorWidget.h
+++ b/Engine/Source/Editor/NiagaraEditor/Private/SNiagaraEffectEditorWidget.h
@@ -320,7 +320,7 @@ protected:
 public:
 
 
-	TSharedRef<ITableRow> OnGenerateConstantListRow(TSharedPtr<EditorExposedVectorConstant> InItem, const TSharedRef< STableViewBase >& OwnerTable)
+	TSharedRef<ITableRow> OnGenerateUpdateConstantListRow(TSharedPtr<EditorExposedVectorConstant> InItem, const TSharedRef< STableViewBase >& OwnerTable)
 	{
 		return SNew(STableRow<TSharedPtr<FString>>, OwnerTable)
 		.Content()
@@ -330,7 +330,7 @@ public:
 	}
 
 
-	TSharedRef<ITableRow> OnGenerateCurveConstantListRow(TSharedPtr<EditorExposedVectorCurveConstant> InItem, const TSharedRef< STableViewBase >& OwnerTable)
+	TSharedRef<ITableRow> OnGenerateUpdateCurveConstantListRow(TSharedPtr<EditorExposedVectorCurveConstant> InItem, const TSharedRef< STableViewBase >& OwnerTable)
 	{
 		return SNew(STableRow<TSharedPtr<FString>>, OwnerTable)
 			.Content()
@@ -340,6 +340,28 @@ public:
 	}
 
 
+	TSharedRef<ITableRow> OnGenerateSpawnConstantListRow(TSharedPtr<EditorExposedVectorConstant> InItem, const TSharedRef< STableViewBase >& OwnerTable)
+	{
+		return SNew(STableRow<TSharedPtr<FString>>, OwnerTable)
+			.Content()
+			[
+				SNew(SVectorConstantWidget).ConstantName(InItem->ConstName).ScriptProps(&Emitter->GetProperties()->SpawnScriptProps)
+			];
+	}
+
+
+	TSharedRef<ITableRow> OnGenerateSpawnCurveConstantListRow(TSharedPtr<EditorExposedVectorCurveConstant> InItem, const TSharedRef< STableViewBase >& OwnerTable)
+	{
+		return SNew(STableRow<TSharedPtr<FString>>, OwnerTable)
+			.Content()
+			[
+				SNew(SCurveConstantWidget).ConstantName(InItem->ConstName).ScriptProps(&Emitter->GetProperties()->SpawnScriptProps).EmitterProps(Emitter->GetProperties())
+			];
+	}
+
+
+
+
 	void OnEmitterNameChanged(const FText &InText)
 	{
 		Emitter->GetProperties()->EmitterName = InText.ToString();
