Copying //UE4/Dev-Rendering to //UE4/Dev-Main (Source: //UE4/Dev-Rendering @ 3045398)

diff --git a/Engine/Source/Developer/MeshUtilities/Private/MeshUtilities.cpp b/Engine/Source/Developer/MeshUtilities/Private/MeshUtilities.cpp
--- a/Engine/Source/Developer/MeshUtilities/Private/MeshUtilities.cpp
+++ b/Engine/Source/Developer/MeshUtilities/Private/MeshUtilities.cpp
@@ -23,6 +23,7 @@
 #include "MeshBoneReduction.h"
 #include "MeshMergeData.h"
 #include "Editor/EditorPerProjectUserSettings.h"
+#include "GPUSkinVertexFactory.h"
 
 #include "Landscape.h"
 #include "LandscapeProxy.h"
@@ -4551,8 +4552,8 @@ public:
 		SkeletalMeshTools::BuildSkeletalMeshChunks(BuildData.Faces, RawVertices, VertIndexAndZ, BuildData.BuildOptions.bKeepOverlappingVertices, BuildData.Chunks, BuildData.bTooManyVerts);
 
 		// Chunk vertices to satisfy the requested limit.
-		static const auto MaxBonesVar = IConsoleManager::Get().FindTConsoleVariableDataInt(TEXT("Compat.MAX_GPUSKIN_BONES"));
-		const int32 MaxGPUSkinBones = MaxBonesVar->GetValueOnAnyThread();
+		const uint32 MaxGPUSkinBones = FGPUBaseSkinVertexFactory::GetMaxGPUSkinBones();
+		check(MaxGPUSkinBones <= FGPUBaseSkinVertexFactory::GHardwareMaxGPUSkinBones);
 		SkeletalMeshTools::ChunkSkinnedVertices(BuildData.Chunks, MaxGPUSkinBones);
 
 		EndSlowTask();
diff --git a/Engine/Source/Developer/ShaderFormatOpenGL/Private/OpenGLShaderCompiler.cpp b/Engine/Source/Developer/ShaderFormatOpenGL/Private/OpenGLShaderCompiler.cpp
--- a/Engine/Source/Developer/ShaderFormatOpenGL/Private/OpenGLShaderCompiler.cpp
+++ b/Engine/Source/Developer/ShaderFormatOpenGL/Private/OpenGLShaderCompiler.cpp
@@ -819,7 +819,8 @@ static void BuildShaderOutput(
 		BuildResourceTableTokenStream(GenericSRT.UnorderedAccessViewMap, GenericSRT.MaxBoundResourceTable, Header.Bindings.ShaderResourceTable.UnorderedAccessViewMap);
 	}
 
-	const int32 MaxSamplers = GetFeatureLevelMaxTextureSamplers(GetMaxSupportedFeatureLevel((EShaderPlatform)ShaderOutput.Target.Platform));
+	// Assume that GL4.3 targets support 32 samplers as we don't currently support separate sampler objects
+	const int32 MaxSamplers = (Version == GLSL_430) ? 32 : GetFeatureLevelMaxTextureSamplers(GetMaxSupportedFeatureLevel((EShaderPlatform)ShaderOutput.Target.Platform));
 
 	if (Header.Bindings.NumSamplers > MaxSamplers)
 	{
diff --git a/Engine/Source/Developer/ShaderFormatOpenGL/Private/ShaderFormatOpenGL.cpp b/Engine/Source/Developer/ShaderFormatOpenGL/Private/ShaderFormatOpenGL.cpp
--- a/Engine/Source/Developer/ShaderFormatOpenGL/Private/ShaderFormatOpenGL.cpp
+++ b/Engine/Source/Developer/ShaderFormatOpenGL/Private/ShaderFormatOpenGL.cpp
@@ -27,7 +27,7 @@ class FShaderFormatGLSL : public IShaderFormat
 		/** Version for shader format, this becomes part of the DDC key. */
 		UE_SHADER_GLSL_150_VER = 60,
 		UE_SHADER_GLSL_150_MAC_VER = 60,
-		UE_SHADER_GLSL_430_VER = 60,
+		UE_SHADER_GLSL_430_VER = 61,
 		UE_SHADER_GLSL_ES2_VER = 60,
 		UE_SHADER_GLSL_150ES2_VER = 60,
 		UE_SHADER_GLSL_150ES2NOUB_VER = 60,
diff --git a/Engine/Source/Developer/TextureCompressor/Private/TextureCompressorModule.cpp b/Engine/Source/Developer/TextureCompressor/Private/TextureCompressorModule.cpp
--- a/Engine/Source/Developer/TextureCompressor/Private/TextureCompressorModule.cpp
+++ b/Engine/Source/Developer/TextureCompressor/Private/TextureCompressorModule.cpp
@@ -288,6 +288,78 @@ private:
 	float KernelWeights[MaxKernelExtend * MaxKernelExtend];
 };
 
+template <EMipGenAddressMode AddressMode>
+static FVector4 ComputeAlphaCoverage(const FVector4& Thresholds, const FVector4& Scales, const FImageView2D& SourceImageData)
+{
+	FVector4 Coverage(0, 0, 0, 0);
+
+	for (int32 y = 0; y < SourceImageData.SizeY; ++y)
+	{
+		for (int32 x = 0; x < SourceImageData.SizeX; ++x)
+		{
+			// Sample channel values at pixel neighborhood
+			FVector4 PixelValue (LookupSourceMip<AddressMode>(SourceImageData, x, y));
+
+			// Calculate coverage for each channel (if being used as an alpha mask)
+			for (int32 i = 0; i < 4; ++i)
+			{
+				// Skip channel if Threshold is 0
+				if (Thresholds[i] == 0)
+				{
+					continue;
+				}
+
+				if (PixelValue[i] * Scales[i] >= Thresholds[i])
+				{
+					++Coverage[i];
+				}
+			}
+		}
+	}
+
+	return Coverage / float(SourceImageData.SizeX * SourceImageData.SizeY);
+}
+
+template <EMipGenAddressMode AddressMode>
+static FVector4 ComputeAlphaScale(const FVector4& Coverages, const FVector4& AlphaThresholds, const FImageView2D& SourceImageData)
+{
+	FVector4 MinAlphaScales (0, 0, 0, 0);
+	FVector4 MaxAlphaScales (4, 4, 4, 4);
+	FVector4 AlphaScales (1, 1, 1, 1);
+
+	//Binary Search to find Alpha Scale
+	for (int32 i = 0; i < 8; ++i)
+	{
+		FVector4 ComputedCoverages = ComputeAlphaCoverage<AddressMode>(AlphaThresholds, AlphaScales, SourceImageData);
+
+		for (int32 j = 0; j < 4; ++j)
+		{
+			if (AlphaThresholds[j] == 0 || fabs(ComputedCoverages[j] - Coverages[j]) < KINDA_SMALL_NUMBER)
+			{
+				continue;
+			}
+
+			if (ComputedCoverages[j] < Coverages[j])
+			{
+				MinAlphaScales[j] = AlphaScales[j];
+			}
+			else if (ComputedCoverages[j] > Coverages[j])
+			{
+				MaxAlphaScales[j] = AlphaScales[j];
+			}
+
+			AlphaScales[j] = (MinAlphaScales[j] + MaxAlphaScales[j]) * 0.5f;
+		}
+
+		if (ComputedCoverages.Equals(Coverages))
+		{
+			break;
+		}
+	}
+
+	return AlphaScales;
+}
+
 
 /**
 * Generates a mip-map for an 2D B8G8R8A8 image using a 4x4 filter with sharpening
@@ -303,6 +375,8 @@ static void GenerateSharpenedMipB8G8R8A8Templ(
 	const FImageView2D& SourceImageData, 
 	FImageView2D& DestImageData, 
 	bool bDitherMipMapAlpha,
+	FVector4 AlphaCoverages,
+	FVector4 AlphaThresholds,
 	const FImageKernel2D& Kernel,
 	uint32 ScaleFactor,
 	bool bSharpenWithoutColorShift )
@@ -316,6 +390,12 @@ static void GenerateSharpenedMipB8G8R8A8Templ(
 	// Set up a random number stream for dithering.
 	FRandomStream RandomStream(0);
 
+	FVector4 AlphaScale(1, 1, 1, 1);
+	if (AlphaThresholds != FVector4(0,0,0,0))
+	{
+		AlphaScale = ComputeAlphaScale<AddressMode>(AlphaCoverages, AlphaThresholds, SourceImageData);
+	}
+	
 	for ( int32 DestY = 0;DestY < DestImageData.SizeY; DestY++ )
 	{
 		for ( int32 DestX = 0;DestX < DestImageData.SizeX; DestX++ )
@@ -374,14 +454,21 @@ static void GenerateSharpenedMipB8G8R8A8Templ(
 				}
 			}
 
+			// Apply computed alpha scales to each channel		
+			FilteredColor.R *= AlphaScale.X;
+			FilteredColor.G *= AlphaScale.Y;
+			FilteredColor.B *= AlphaScale.Z;
+			FilteredColor.A *= AlphaScale.W;
+
+
 			if ( bDitherMipMapAlpha )
 			{
 				// Dither the alpha of any pixel which passes an alpha threshold test.
-				const int32 AlphaThreshold = 5.0f / 255.0f;
+				const int32 DitherAlphaThreshold = 5.0f / 255.0f;
 				const float MinRandomAlpha = 85.0f;
 				const float MaxRandomAlpha = 255.0f;
 
-				if ( FilteredColor.A > AlphaThreshold )
+				if ( FilteredColor.A > DitherAlphaThreshold)
 				{
 					FilteredColor.A = FMath::TruncToInt( FMath::Lerp( MinRandomAlpha, MaxRandomAlpha, RandomStream.GetFraction() ) );
 				}
@@ -402,6 +489,8 @@ static void GenerateSharpenedMipB8G8R8A8(
 	FImageView2D& DestImageData, 
 	EMipGenAddressMode AddressMode, 
 	bool bDitherMipMapAlpha,
+	FVector4 AlphaCoverages,
+	FVector4 AlphaThresholds,
 	const FImageKernel2D &Kernel,
 	uint32 ScaleFactor,
 	bool bSharpenWithoutColorShift
@@ -410,13 +499,13 @@ static void GenerateSharpenedMipB8G8R8A8(
 	switch(AddressMode)
 	{
 	case MGTAM_Wrap:
-		GenerateSharpenedMipB8G8R8A8Templ<MGTAM_Wrap>(SourceImageData, DestImageData, bDitherMipMapAlpha, Kernel, ScaleFactor, bSharpenWithoutColorShift);
+		GenerateSharpenedMipB8G8R8A8Templ<MGTAM_Wrap>(SourceImageData, DestImageData, bDitherMipMapAlpha, AlphaCoverages, AlphaThresholds, Kernel, ScaleFactor, bSharpenWithoutColorShift);
 		break;
 	case MGTAM_Clamp:
-		GenerateSharpenedMipB8G8R8A8Templ<MGTAM_Clamp>(SourceImageData, DestImageData, bDitherMipMapAlpha, Kernel, ScaleFactor, bSharpenWithoutColorShift);
+		GenerateSharpenedMipB8G8R8A8Templ<MGTAM_Clamp>(SourceImageData, DestImageData, bDitherMipMapAlpha, AlphaCoverages, AlphaThresholds, Kernel, ScaleFactor, bSharpenWithoutColorShift);
 		break;
 	case MGTAM_BorderBlack:
-		GenerateSharpenedMipB8G8R8A8Templ<MGTAM_BorderBlack>(SourceImageData, DestImageData, bDitherMipMapAlpha, Kernel, ScaleFactor, bSharpenWithoutColorShift);
+		GenerateSharpenedMipB8G8R8A8Templ<MGTAM_BorderBlack>(SourceImageData, DestImageData, bDitherMipMapAlpha, AlphaCoverages, AlphaThresholds, Kernel, ScaleFactor, bSharpenWithoutColorShift);
 		break;
 	default:
 		check(0);
@@ -514,6 +603,8 @@ static void GenerateTopMip(const FImage& SrcImage, FImage& DestImage, const FTex
 			DestView,
 			AddressMode,
 			Settings.bDitherMipMapAlpha,
+			FVector4(0, 0, 0, 0),
+			FVector4(0, 0, 0, 0),
 			KernelDownsample,
 			1,
 			Settings.bSharpenWithoutColorShift
@@ -542,6 +633,8 @@ static void GenerateMipChain(
 	const int32 SrcHeight= BaseMip.SizeY;
 	const int32 SrcNumSlices = BaseMip.NumSlices;
 	const ERawImageFormat::Type ImageFormat = ERawImageFormat::RGBA32F;
+	FVector4 AlphaScales(1, 1, 1, 1);
+	FVector4 AlphaCoverages(0, 0, 0, 0);
 
 	// space for one source mip and one destination mip
 	FImage IntermediateSrc(SrcWidth, SrcHeight, SrcNumSlices, ImageFormat);
@@ -563,6 +656,26 @@ static void GenerateMipChain(
 		bReDrawBorder = !Settings.bBorderColorBlack;
 	}
 
+	// Calculate alpha coverage value to preserve along mip chain
+	if (Settings.AlphaCoverageThresholds != FVector4(0,0,0,0))
+	{
+		FImageView2D IntermediateSrcView(IntermediateSrc, 0);
+		switch (AddressMode)
+		{
+		case MGTAM_Wrap:
+			AlphaCoverages = ComputeAlphaCoverage<MGTAM_Wrap>(Settings.AlphaCoverageThresholds, AlphaScales, IntermediateSrcView);
+			break;
+		case MGTAM_Clamp:
+			AlphaCoverages = ComputeAlphaCoverage<MGTAM_Clamp>(Settings.AlphaCoverageThresholds, AlphaScales, IntermediateSrcView);
+			break;
+		case MGTAM_BorderBlack:
+			AlphaCoverages = ComputeAlphaCoverage<MGTAM_BorderBlack>(Settings.AlphaCoverageThresholds, AlphaScales, IntermediateSrcView);
+			break;
+		default:
+			check(0);
+		}		
+	}
+
 	// Generate mips
 	for (; MipChainDepth != 0 ; --MipChainDepth)
 	{
@@ -580,6 +693,8 @@ static void GenerateMipChain(
 				DestView,
 				AddressMode,
 				Settings.bDitherMipMapAlpha,
+				AlphaCoverages,
+				Settings.AlphaCoverageThresholds,
 				KernelDownsample,
 				2,
 				Settings.bSharpenWithoutColorShift
@@ -594,6 +709,8 @@ static void GenerateMipChain(
 					IntermediateDstView,
 					AddressMode,
 					Settings.bDitherMipMapAlpha,
+					AlphaCoverages,
+					Settings.AlphaCoverageThresholds,
 					KernelSimpleAverage,
 					2,
 					Settings.bSharpenWithoutColorShift
diff --git a/Engine/Source/Developer/TextureCompressor/Public/TextureCompressorModule.h b/Engine/Source/Developer/TextureCompressor/Public/TextureCompressorModule.h
--- a/Engine/Source/Developer/TextureCompressor/Public/TextureCompressorModule.h
+++ b/Engine/Source/Developer/TextureCompressor/Public/TextureCompressorModule.h
@@ -65,6 +65,8 @@ struct FTextureBuildSettings
 {
 	/** Color adjustment parameters. */
 	FColorAdjustmentParameters ColorAdjustment;
+	/** Channel values to compare to when preserving alpha coverage. */
+	FVector4 AlphaCoverageThresholds;
 	/** The desired amount of mip sharpening. */
 	float MipSharpening;
 	/** For angular filtered cubemaps, the mip level which contains convolution with the diffuse cosine lobe. */
@@ -130,7 +132,8 @@ struct FTextureBuildSettings
 
 	/** Default settings. */
 	FTextureBuildSettings()
-		: MipSharpening( 0.0f )
+		: AlphaCoverageThresholds(0, 0, 0, 0)
+		, MipSharpening( 0.0f )
 		, DiffuseConvolveMipLevel( 0 )
 		, SharpenMipKernelSize( 2 )
 		, MaxTextureResolution(TNumericLimits<uint32>::Max())
diff --git a/Engine/Source/Editor/GraphEditor/Private/NodeFactory.cpp b/Engine/Source/Editor/GraphEditor/Private/NodeFactory.cpp
--- a/Engine/Source/Editor/GraphEditor/Private/NodeFactory.cpp
+++ b/Engine/Source/Editor/GraphEditor/Private/NodeFactory.cpp
@@ -87,6 +87,8 @@
 
 #include "EdGraphUtilities.h"
 
+#include "EdGraphSchema_Niagara.h"
+
 TSharedPtr<SGraphNode> FNodeFactory::CreateNodeWidget(UEdGraphNode* InNode)
 {
 	check(InNode != NULL);
@@ -423,29 +425,7 @@ TSharedPtr<SGraphPin> FNodeFactory::CreatePinWidget(UEdGraphPin* InPin)
 			return SNew(SGraphPin, InPin);
 		}
 	}
-
-	if (const UEdGraphSchema_Niagara* NSchema = Cast<const UEdGraphSchema_Niagara>(InPin->GetSchema()))
-	{
-		if (InPin->PinType.PinCategory == NSchema->PC_Float)
-		{
-			return SNew(SGraphPinNum, InPin);
-		}
-		if (InPin->PinType.PinCategory == NSchema->PC_Vector)
-		{
-			return SNew(SGraphPinVector4, InPin);
-		}
-		if (InPin->PinType.PinCategory == NSchema->PC_Matrix)
-		{
-			return SNew(SGraphPin, InPin);
-		}
-		if (InPin->PinType.PinCategory == NSchema->PC_Curve)
-		{
-			return SNew(SGraphPin, InPin);
-		}
-
-	}
-
-
+	
 	// If we didn't pick a custom pin widget, use an uncustomized basic pin
 	return SNew(SGraphPin, InPin);
 }
@@ -503,3 +483,26 @@ FConnectionDrawingPolicy* FNodeFactory::CreateConnectionPolicy(const UEdGraphSch
     // If we never picked a custom policy, use the uncustomized standard policy
     return ConnectionDrawingPolicy;
 }
+
+//////////////////////////////////////////////////////////////////////////
+
+class FNiagaraScriptGraphPanelPinFactory : public FGraphPanelPinFactory
+{
+	virtual TSharedPtr<class SGraphPin> CreatePin(class UEdGraphPin* InPin) const override
+	{
+		const UEdGraphSchema_Niagara* NSchema = GetDefault<UEdGraphSchema_Niagara>();
+		if (InPin->PinType.PinCategory == NSchema->GetFloatName())
+		{
+			return SNew(SGraphPinNum, InPin);
+		}
+		if (InPin->PinType.PinCategory == NSchema->GetVectorName())
+		{
+			return SNew(SGraphPinVector4, InPin);
+		}
+		if (InPin->PinType.PinCategory == NSchema->GetMatrixName())
+		{
+			return SNew(SGraphPin, InPin);
+		}
+		return NULL;
+	}
+};
\ No newline at end of file
diff --git a/Engine/Source/Editor/NiagaraEditor/NiagaraEditor.Build.cs b/Engine/Source/Editor/NiagaraEditor/NiagaraEditor.Build.cs
--- a/Engine/Source/Editor/NiagaraEditor/NiagaraEditor.Build.cs
+++ b/Engine/Source/Editor/NiagaraEditor/NiagaraEditor.Build.cs
@@ -31,7 +31,8 @@ public class NiagaraEditor : ModuleRules
             new string[] {
 				"Engine",
                 "Niagara",
-				"PropertyEditor",
+                "UnrealEd",
+                "PropertyEditor",
             }
         );
 
diff --git a/Engine/Source/Editor/NiagaraEditor/Private/EdGraphSchema_Niagara.cpp b/Engine/Source/Editor/NiagaraEditor/Private/EdGraphSchema_Niagara.cpp
new file mode 100644
--- /dev/null
+++ b/Engine/Source/Editor/NiagaraEditor/Private/EdGraphSchema_Niagara.cpp
@@ -0,0 +1,461 @@
+// Copyright 1998-2016 Epic Games, Inc. All Rights Reserved.
+
+#include "NiagaraEditorPrivatePCH.h"
+#include "BlueprintGraphDefinitions.h"
+#include "GraphEditorSettings.h"
+#include "NiagaraComponent.h"
+#include "GraphEditorActions.h"
+#include "EdGraphSchema_Niagara.h"
+
+#include "INiagaraEditor.h"
+#include "INiagaraCompiler.h"
+#include "ScopedTransaction.h"
+
+#include "NiagaraNodeOutput.h"
+#include "NiagaraNodeInput.h"
+#include "NiagaraNodeFunctionCall.h"
+#include "NiagaraNodeReadDataSet.h"
+#include "NiagaraNodeWriteDataSet.h"
+#include "NiagaraNodeOp.h"
+
+#define LOCTEXT_NAMESPACE "NiagaraSchema"
+
+#define SNAP_GRID (16) // @todo ensure this is the same as SNodePanel::GetSnapGridSize()
+
+const FLinearColor UEdGraphSchema_Niagara::NodeTitleColor_Attribute = FLinearColor::Green;
+const FLinearColor UEdGraphSchema_Niagara::NodeTitleColor_Constant = FLinearColor::Red;
+const FLinearColor UEdGraphSchema_Niagara::NodeTitleColor_SystemConstant = FLinearColor::White;
+const FLinearColor UEdGraphSchema_Niagara::NodeTitleColor_FunctionCall = FLinearColor::Blue;
+const FLinearColor UEdGraphSchema_Niagara::NodeTitleColor_Event = FLinearColor::Red;
+
+namespace 
+{
+	// Maximum distance a drag can be off a node edge to require 'push off' from node
+	const int32 NodeDistance = 60;
+}
+
+UEdGraphNode* FNiagaraSchemaAction_NewNode::PerformAction(class UEdGraph* ParentGraph, UEdGraphPin* FromPin, const FVector2D Location, bool bSelectNewNode/* = true*/)
+{
+	UEdGraphNode* ResultNode = NULL;
+
+	// If there is a template, we actually use it
+	if (NodeTemplate != NULL)
+	{
+		const FScopedTransaction Transaction(NSLOCTEXT("UnrealEd", "NiagaraEditorNewNode", "Niagara Editor: New Node"));
+		ParentGraph->Modify();
+
+		NodeTemplate->SetFlags(RF_Transactional);
+
+		// set outer to be the graph so it doesn't go away
+		NodeTemplate->Rename(NULL, ParentGraph, REN_NonTransactional);
+		ParentGraph->AddNode(NodeTemplate, true, bSelectNewNode);
+
+		NodeTemplate->CreateNewGuid();
+		NodeTemplate->PostPlacedNewNode();
+		NodeTemplate->AllocateDefaultPins();
+		NodeTemplate->AutowireNewNode(FromPin);
+
+		// For input pins, new node will generally overlap node being dragged off
+		// Work out if we want to visually push away from connected node
+		int32 XLocation = Location.X;
+		if (FromPin && FromPin->Direction == EGPD_Input)
+		{
+			UEdGraphNode* PinNode = FromPin->GetOwningNode();
+			const float XDelta = FMath::Abs(PinNode->NodePosX - Location.X);
+
+			if (XDelta < NodeDistance)
+			{
+				// Set location to edge of current node minus the max move distance
+				// to force node to push off from connect node enough to give selection handle
+				XLocation = PinNode->NodePosX - NodeDistance;
+			}
+		}
+
+		NodeTemplate->NodePosX = XLocation;
+		NodeTemplate->NodePosY = Location.Y;
+		NodeTemplate->SnapToGrid(SNAP_GRID);
+
+		ResultNode = NodeTemplate;
+
+		ParentGraph->NotifyGraphChanged();
+	}
+
+	return ResultNode;
+}
+
+UEdGraphNode* FNiagaraSchemaAction_NewNode::PerformAction(class UEdGraph* ParentGraph, TArray<UEdGraphPin*>& FromPins, const FVector2D Location, bool bSelectNewNode/* = true*/) 
+{
+	UEdGraphNode* ResultNode = NULL;
+
+	if (FromPins.Num() > 0)
+	{
+		ResultNode = PerformAction(ParentGraph, FromPins[0], Location, bSelectNewNode);
+
+		// Try autowiring the rest of the pins
+		for (int32 Index = 1; Index < FromPins.Num(); ++Index)
+		{
+			ResultNode->AutowireNewNode(FromPins[Index]);
+		}
+	}
+	else
+	{
+		ResultNode = PerformAction(ParentGraph, NULL, Location, bSelectNewNode);
+	}
+
+	return ResultNode;
+}
+
+void FNiagaraSchemaAction_NewNode::AddReferencedObjects( FReferenceCollector& Collector )
+{
+	FEdGraphSchemaAction::AddReferencedObjects( Collector );
+
+	// These don't get saved to disk, but we want to make sure the objects don't get GC'd while the action array is around
+	Collector.AddReferencedObject( NodeTemplate );
+}
+
+//////////////////////////////////////////////////////////////////////////
+
+UEdGraphSchema_Niagara::UEdGraphSchema_Niagara(const FObjectInitializer& ObjectInitializer)
+	: Super(ObjectInitializer)
+{
+
+	PC_Float = TEXT("float");
+	PC_Vector = TEXT("vector");
+	PC_Matrix = TEXT("matrix");
+}
+
+TSharedPtr<FNiagaraSchemaAction_NewNode> AddNewNodeAction(FGraphContextMenuBuilder& ContextMenuBuilder, const FText& Category, const FText& MenuDesc, const FString& Tooltip)
+{
+	TSharedPtr<FNiagaraSchemaAction_NewNode> NewAction = TSharedPtr<FNiagaraSchemaAction_NewNode>(new FNiagaraSchemaAction_NewNode(Category, MenuDesc, Tooltip, 0));
+	ContextMenuBuilder.AddAction( NewAction );
+	return NewAction;
+}
+
+void UEdGraphSchema_Niagara::GetGraphContextActions(FGraphContextMenuBuilder& ContextMenuBuilder) const
+{
+	const UNiagaraGraph* NiagaraGraph = CastChecked<UNiagaraGraph>(ContextMenuBuilder.CurrentGraph);
+
+	//Add menu options for Attributes currently defined by the output node.
+	const UNiagaraNodeOutput* OutNode = NiagaraGraph->FindOutputNode();
+	if (OutNode)
+	{
+		for (int32 i = 0; i < OutNode->Outputs.Num(); i++)
+		{
+			const FNiagaraVariableInfo Attr = OutNode->Outputs[i];
+
+			FFormatNamedArguments Args;
+			Args.Add(TEXT("Attribute"), FText::FromName(Attr.Name));
+			const FText MenuDesc = FText::Format(LOCTEXT("GetAttribute", "Get {Attribute}"), Args);
+
+			TSharedPtr<FNiagaraSchemaAction_NewNode> GetAttrAction = AddNewNodeAction(ContextMenuBuilder, LOCTEXT("Attributes Menu Title", "Attributes"), MenuDesc, TEXT(""));
+
+			UNiagaraNodeInput* InputNode = NewObject<UNiagaraNodeInput>(ContextMenuBuilder.OwnerOfTemporaries);
+			InputNode->Input = Attr;
+			GetAttrAction->NodeTemplate = InputNode;
+		}
+	}
+
+#define NiagaraOp(OPNAME) \
+		if(const FNiagaraOpInfo* OpInfo = FNiagaraOpInfo::GetOpInfo(FNiagaraOpInfo::OPNAME))\
+		{\
+		TSharedPtr<FNiagaraSchemaAction_NewNode> AddOpAction = AddNewNodeAction(ContextMenuBuilder, LOCTEXT("Operations Menu Title", "Operations"), OpInfo->FriendlyName, TEXT("")); \
+			UNiagaraNodeOp* OpNode = NewObject<UNiagaraNodeOp>(ContextMenuBuilder.OwnerOfTemporaries); \
+			OpNode->OpName = OpInfo->Name; \
+			AddOpAction->NodeTemplate = OpNode; \
+		}
+
+	NiagaraOpList
+
+#undef NiagaraOp
+
+	//Emitter constants managed by the system.
+	const TArray<FNiagaraVariableInfo>& SystemConstants = UNiagaraComponent::GetSystemConstants();
+	for (const FNiagaraVariableInfo& SysConst : SystemConstants)
+	{
+		FFormatNamedArguments Args;
+		Args.Add(TEXT("Constant"), FText::FromName(SysConst.Name));
+		const FText MenuDesc = FText::Format(LOCTEXT("GetSystemConstant", "Get {Constant}"), Args);
+
+		TSharedPtr<FNiagaraSchemaAction_NewNode> GetConstAction = AddNewNodeAction(ContextMenuBuilder, LOCTEXT("System Constants Menu Title", "System Constants"), MenuDesc, TEXT(""));
+
+		UNiagaraNodeInput* InputNode = NewObject<UNiagaraNodeInput>(ContextMenuBuilder.OwnerOfTemporaries);
+		InputNode->Input = SysConst;
+		InputNode->bExposeWhenConstant = false;
+		InputNode->bCanBeExposed = false;
+		GetConstAction->NodeTemplate = InputNode;
+	}
+
+	//Add a generic Input node to allow getting external constants the current context doesn't know about.
+	{
+		const FText MenuDesc = LOCTEXT("GetInput", "Input");
+
+		//TSharedPtr<FNiagaraSchemaAction_NewNode> GetConstAction = AddNewNodeAction(ContextMenuBuilder, *LOCTEXT("System Constants Menu Title", "System Constants").ToString(), MenuDesc, TEXT(""));
+		TSharedPtr<FNiagaraSchemaAction_NewNode> InputAction = AddNewNodeAction(ContextMenuBuilder, FText::GetEmpty(), MenuDesc, TEXT(""));
+
+		UNiagaraNodeInput* InputNode = NewObject<UNiagaraNodeInput>(ContextMenuBuilder.OwnerOfTemporaries);
+		InputNode->Input.Name = NAME_None;
+		InputNode->Input.Type = ENiagaraDataType::Vector;
+		InputAction->NodeTemplate = InputNode;
+	}
+
+	//Add a function call node
+	{
+		const FText MenuDesc = LOCTEXT("NiagaraFunctionCall", "Function Call");
+
+		TSharedPtr<FNiagaraSchemaAction_NewNode> FunctionCallAction = AddNewNodeAction(ContextMenuBuilder, FText::GetEmpty(), MenuDesc, TEXT(""));
+
+		UNiagaraNodeFunctionCall* FunctionCallNode = NewObject<UNiagaraNodeFunctionCall>(ContextMenuBuilder.OwnerOfTemporaries);
+		FunctionCallAction->NodeTemplate = FunctionCallNode;
+	}
+
+	//Add event read and writes nodes
+	{
+		const FText MenuDesc = LOCTEXT("NiagaraEventRead", "Event Read");
+
+		TSharedPtr<FNiagaraSchemaAction_NewNode> Action = AddNewNodeAction(ContextMenuBuilder, FText::GetEmpty(), MenuDesc, TEXT(""));
+
+		UNiagaraNodeReadDataSet* EventReadNode = NewObject<UNiagaraNodeReadDataSet>(ContextMenuBuilder.OwnerOfTemporaries);
+		Action->NodeTemplate = EventReadNode;
+	}	
+	{
+		const FText MenuDesc = LOCTEXT("NiagaraEventWrite", "Event Write");
+
+		TSharedPtr<FNiagaraSchemaAction_NewNode> Action = AddNewNodeAction(ContextMenuBuilder, FText::GetEmpty(), MenuDesc, TEXT(""));
+
+		UNiagaraNodeWriteDataSet* EventWriteNode = NewObject<UNiagaraNodeWriteDataSet>(ContextMenuBuilder.OwnerOfTemporaries);
+		Action->NodeTemplate = EventWriteNode;
+	}
+	//TODO: Add quick commands for certain UNiagaraStructs and UNiagaraScripts to be added as functions
+}
+
+const FPinConnectionResponse UEdGraphSchema_Niagara::CanCreateConnection(const UEdGraphPin* PinA, const UEdGraphPin* PinB) const
+{
+	// Make sure the pins are not on the same node
+	if (PinA->GetOwningNode() == PinB->GetOwningNode())
+	{
+		return FPinConnectionResponse(CONNECT_RESPONSE_DISALLOW, TEXT("Both are on the same node"));
+	}
+
+	// Check both pins support connections
+	if(PinA->bNotConnectable || PinB->bNotConnectable)
+	{
+		return FPinConnectionResponse(CONNECT_RESPONSE_DISALLOW, TEXT("Pin doesn't support connections."));
+	}
+
+	// Compare the directions
+	const UEdGraphPin* InputPin = NULL;
+	const UEdGraphPin* OutputPin = NULL;
+
+	if (!CategorizePinsByDirection(PinA, PinB, /*out*/ InputPin, /*out*/ OutputPin))
+	{
+		return FPinConnectionResponse(CONNECT_RESPONSE_DISALLOW, TEXT("Directions are not compatible"));
+	}
+
+	// Types must match exactly
+	if(PinA->PinType != PinB->PinType)
+	{
+		return FPinConnectionResponse(CONNECT_RESPONSE_DISALLOW, TEXT("Types are not compatible"));
+	}
+
+	// See if we want to break existing connections (if its an input with an existing connection)
+	const bool bBreakExistingDueToDataInput = (InputPin->LinkedTo.Num() > 0);
+	if (bBreakExistingDueToDataInput)
+	{
+		const ECanCreateConnectionResponse ReplyBreakInputs = (PinA == InputPin) ? CONNECT_RESPONSE_BREAK_OTHERS_A : CONNECT_RESPONSE_BREAK_OTHERS_B;
+		return FPinConnectionResponse(ReplyBreakInputs, TEXT("Replace existing input connections"));
+	}
+	else
+	{
+		return FPinConnectionResponse(CONNECT_RESPONSE_MAKE, TEXT(""));
+	}
+}
+
+void UEdGraphSchema_Niagara::BreakSinglePinLink(UEdGraphPin* SourcePin, UEdGraphPin* TargetPin) 
+{
+	const FScopedTransaction Transaction(NSLOCTEXT("UnrealEd", "NiagaraEditorBreakConnection", "Niagara Editor: Break Connection"));
+
+	Super::BreakSinglePinLink(SourcePin, TargetPin);
+}
+
+bool UEdGraphSchema_Niagara::TryCreateConnection(UEdGraphPin* A, UEdGraphPin* B) const
+{
+	const FScopedTransaction Transaction(NSLOCTEXT("UnrealEd", "NiagaraEditorCreateConnection", "Niagara Editor: Create Connection"));
+
+	return Super::TryCreateConnection(A, B);
+}
+
+FLinearColor UEdGraphSchema_Niagara::GetPinTypeColor(const FEdGraphPinType& PinType) const
+{
+	const FString& TypeString = PinType.PinCategory;
+	const UGraphEditorSettings* Settings = GetDefault<UGraphEditorSettings>();
+
+	if (TypeString == PC_Float)
+	{
+		return Settings->FloatPinTypeColor;
+	}
+	else if (TypeString == PC_Vector)
+	{
+		return Settings->VectorPinTypeColor;
+	}
+	else if (TypeString == PC_Matrix)
+	{
+		return Settings->TransformPinTypeColor;
+	}
+
+	return Settings->DefaultPinTypeColor;
+}
+
+bool UEdGraphSchema_Niagara::ShouldHidePinDefaultValue(UEdGraphPin* Pin) const
+{
+	check(Pin != NULL);
+
+	if (Pin->bDefaultValueIsIgnored)
+	{
+		return true;
+	}
+
+	return false;
+}
+
+ENiagaraDataType UEdGraphSchema_Niagara::GetPinDataType(UEdGraphPin* Pin) const
+{
+	if (Pin->PinType.PinCategory == PC_Float)
+	{
+		return ENiagaraDataType::Scalar;
+	}
+	else if (Pin->PinType.PinCategory == PC_Vector)
+	{
+		return ENiagaraDataType::Vector;
+	}
+	else if (Pin->PinType.PinCategory == PC_Matrix)
+	{
+		return ENiagaraDataType::Matrix;
+	}
+
+	check(0);
+	return ENiagaraDataType::Vector;
+}
+
+void UEdGraphSchema_Niagara::GetPinDefaultValue(UEdGraphPin* Pin, float& OutDefault)const
+{
+	//Ugh, this string parsing is rubbish. Surely there's a consistent api for dealing with default values in the edgraph without all this?
+	FString ResultString = Pin->DefaultValue.IsEmpty() ? Pin->AutogeneratedDefaultValue : Pin->DefaultValue;
+	ResultString.Trim();
+	ResultString.TrimTrailing();
+	OutDefault = FCString::Atof(*ResultString);
+}
+
+void UEdGraphSchema_Niagara::GetPinDefaultValue(UEdGraphPin* Pin, FVector4& OutDefault)const
+{
+	//Ugh, this string parsing is rubbish. Surely there's a consistent api for dealing with default values in the edgraph without all this?
+	TArray<FString> ResultString;
+	FString ValueString = Pin->DefaultValue.IsEmpty() ? Pin->AutogeneratedDefaultValue : Pin->DefaultValue;
+	ValueString.Trim();
+	ValueString.TrimTrailing();
+	ValueString.ParseIntoArray(ResultString, TEXT(","), true);
+	if (ResultString.Num() == 4)
+	{
+		OutDefault = FVector4(FCString::Atof(*ResultString[0]), FCString::Atof(*ResultString[1]), FCString::Atof(*ResultString[2]), FCString::Atof(*ResultString[3]));
+	}
+}
+
+void UEdGraphSchema_Niagara::GetPinDefaultValue(UEdGraphPin* Pin, FMatrix& OutDefault)const
+{
+	TArray<FString> ResultString;
+	FString ValueString = Pin->DefaultValue.IsEmpty() ? Pin->AutogeneratedDefaultValue : Pin->DefaultValue;
+	ValueString.Trim();
+	ValueString.TrimTrailing();
+	ValueString.ParseIntoArray(ResultString, TEXT(","), true);
+	check(ResultString.Num() == 16);
+	if (ResultString.Num() == 16)
+	{
+		OutDefault.M[0][0] = FCString::Atof(*ResultString[0]); OutDefault.M[0][1] = FCString::Atof(*ResultString[1]); OutDefault.M[0][2] = FCString::Atof(*ResultString[2]); OutDefault.M[0][3] = FCString::Atof(*ResultString[3]);
+		OutDefault.M[1][0] = FCString::Atof(*ResultString[4]); OutDefault.M[1][1] = FCString::Atof(*ResultString[5]); OutDefault.M[1][2] = FCString::Atof(*ResultString[6]); OutDefault.M[1][3] = FCString::Atof(*ResultString[7]);
+		OutDefault.M[2][0] = FCString::Atof(*ResultString[8]); OutDefault.M[2][1] = FCString::Atof(*ResultString[9]); OutDefault.M[2][2] = FCString::Atof(*ResultString[10]); OutDefault.M[2][3] = FCString::Atof(*ResultString[11]);
+		OutDefault.M[3][0] = FCString::Atof(*ResultString[12]); OutDefault.M[3][1] = FCString::Atof(*ResultString[13]); OutDefault.M[3][2] = FCString::Atof(*ResultString[14]); OutDefault.M[3][3] = FCString::Atof(*ResultString[15]);
+	}
+}
+
+bool UEdGraphSchema_Niagara::IsSystemConstant(const FNiagaraVariableInfo& Variable)const
+{
+	return UNiagaraComponent::GetSystemConstants().Find(Variable) != INDEX_NONE;
+}
+
+void UEdGraphSchema_Niagara::GetBreakLinkToSubMenuActions(class FMenuBuilder& MenuBuilder, UEdGraphPin* InGraphPin)
+{
+	// Make sure we have a unique name for every entry in the list
+	TMap< FString, uint32 > LinkTitleCount;
+
+	// Add all the links we could break from
+	for (TArray<class UEdGraphPin*>::TConstIterator Links(InGraphPin->LinkedTo); Links; ++Links)
+	{
+		UEdGraphPin* Pin = *Links;
+		FString TitleString = Pin->GetOwningNode()->GetNodeTitle(ENodeTitleType::ListView).ToString();
+		FText Title = FText::FromString(TitleString);
+		if (Pin->PinName != TEXT(""))
+		{
+			TitleString = FString::Printf(TEXT("%s (%s)"), *TitleString, *Pin->PinName);
+
+			// Add name of connection if possible
+			FFormatNamedArguments Args;
+			Args.Add(TEXT("NodeTitle"), Title);
+			Args.Add(TEXT("PinName"), Pin->GetDisplayName());
+			Title = FText::Format(LOCTEXT("BreakDescPin", "{NodeTitle} ({PinName})"), Args);
+		}
+
+		uint32 &Count = LinkTitleCount.FindOrAdd(TitleString);
+
+		FText Description;
+		FFormatNamedArguments Args;
+		Args.Add(TEXT("NodeTitle"), Title);
+		Args.Add(TEXT("NumberOfNodes"), Count);
+
+		if (Count == 0)
+		{
+			Description = FText::Format(LOCTEXT("BreakDesc", "Break link to {NodeTitle}"), Args);
+		}
+		else
+		{
+			Description = FText::Format(LOCTEXT("BreakDescMulti", "Break link to {NodeTitle} ({NumberOfNodes})"), Args);
+		}
+		++Count;
+		MenuBuilder.AddMenuEntry(Description, Description, FSlateIcon(), FUIAction(
+			FExecuteAction::CreateUObject((UEdGraphSchema_Niagara*const)this, &UEdGraphSchema::BreakSinglePinLink, const_cast<UEdGraphPin*>(InGraphPin), *Links)));
+	}
+}
+
+void UEdGraphSchema_Niagara::GetContextMenuActions(const UEdGraph* CurrentGraph, const UEdGraphNode* InGraphNode, const UEdGraphPin* InGraphPin, class FMenuBuilder* MenuBuilder, bool bIsDebugging) const
+{
+	if (InGraphPin)
+	{
+		MenuBuilder->BeginSection("EdGraphSchema_NiagaraPinActions", LOCTEXT("PinActionsMenuHeader", "Pin Actions"));
+		{
+			// Only display the 'Break Link' option if there is a link to break!
+			if (InGraphPin->LinkedTo.Num() > 0)
+			{
+				MenuBuilder->AddMenuEntry(FGraphEditorCommands::Get().BreakPinLinks);
+
+				// add sub menu for break link to
+				if (InGraphPin->LinkedTo.Num() > 1)
+				{
+					MenuBuilder->AddSubMenu(
+						LOCTEXT("BreakLinkTo", "Break Link To..."),
+						LOCTEXT("BreakSpecificLinks", "Break a specific link..."),
+						FNewMenuDelegate::CreateUObject((UEdGraphSchema_Niagara*const)this, &UEdGraphSchema_Niagara::GetBreakLinkToSubMenuActions, const_cast<UEdGraphPin*>(InGraphPin)));
+				}
+				else
+				{
+					((UEdGraphSchema_Niagara*const)this)->GetBreakLinkToSubMenuActions(*MenuBuilder, const_cast<UEdGraphPin*>(InGraphPin));
+				}				
+			}
+		}
+		MenuBuilder->EndSection();
+	}
+	else if (InGraphNode)
+	{
+	}
+
+	Super::GetContextMenuActions(CurrentGraph, InGraphNode, InGraphPin, MenuBuilder, bIsDebugging);
+}
+
+#undef LOCTEXT_NAMESPACE
\ No newline at end of file
diff --git a/Engine/Source/Editor/NiagaraEditor/Private/NiagaraCompiler.cpp b/Engine/Source/Editor/NiagaraEditor/Private/NiagaraCompiler.cpp
--- a/Engine/Source/Editor/NiagaraEditor/Private/NiagaraCompiler.cpp
+++ b/Engine/Source/Editor/NiagaraEditor/Private/NiagaraCompiler.cpp
@@ -9,6 +9,11 @@
 #include "ComponentReregisterContext.h"
 #include "NiagaraCompiler_VectorVM.h"
 #include "NiagaraEditorCommon.h"
+#include "NiagaraNodeFunctionCall.h"
+#include "NiagaraNodeInput.h"
+#include "NiagaraNodeOutput.h"
+
+#include "EdGraphSchema_Niagara.h"
 
 #define LOCTEXT_NAMESPACE "NiagaraCompiler"
 
@@ -72,11 +77,6 @@ public:
 			SourceExpressions.Add(MakeShareable(new FNiagaraExpression(InCompiler, FNiagaraVariableInfo(ResultName, ENiagaraDataType::Vector))));
 			SourceExpressions.Add(MakeShareable(new FNiagaraExpression(InCompiler, FNiagaraVariableInfo(ResultName, ENiagaraDataType::Vector))));
 		}
-// 		else if (Result.Type == ENiagaraDataType::Curve)
-// 		{
-// 			static const FName ResultName(TEXT("Curve"));
-// 			SourceExpressions.Add(MakeShareable(new FNiagaraExpression(InCompiler, FNiagaraVariableInfo(ResultName, ENiagaraDataType::Curve))));
-// 		}
 	}
 
 	virtual bool Process()override
@@ -101,10 +101,6 @@ public:
 			Src3->ResultLocation = ENiagaraExpressionResultLocation::Constants;
 			Src3->ResultIndex = ResultIndex + 3;
 		}
-		else if (Result.Type == ENiagaraDataType::Curve)
-		{
-			ResultLocation = ENiagaraExpressionResultLocation::BufferConstants;
-		}
 
 		return true;
 	}
@@ -296,7 +292,7 @@ TNiagaraExprPtr FNiagaraCompiler::CompilePin(UEdGraphPin* Pin)
 		{
 			//No connections to this input so add the default as a const expression.
 			const UEdGraphSchema_Niagara* Schema = Cast<const UEdGraphSchema_Niagara>(Pin->GetSchema());
-			ENiagaraDataType PinType = Schema->GetPinType(Pin);
+			ENiagaraDataType PinType = Schema->GetPinDataType(Pin);
 
 			FString ConstString = Pin->GetDefaultAsString();
 			FName ConstName(*ConstString);
@@ -326,12 +322,6 @@ TNiagaraExprPtr FNiagaraCompiler::CompilePin(UEdGraphPin* Pin)
 				ConstantData.SetOrAddInternal(Var, Default);
 			}
 				break;
-			case ENiagaraDataType::Curve:
-			{
-				UNiagaraDataObject *Default = nullptr;
-				ConstantData.SetOrAddInternal(Var, Default);
-			}
-				break;
 			}
 			Ret = Expression_GetInternalConstant(Var);
 		}
@@ -398,12 +388,6 @@ TNiagaraExprPtr FNiagaraCompiler::GetExternalConstant(const FNiagaraVariableInfo
 	SetOrAddConstant(false, Constant, Default);
 	return Expression_GetExternalConstant(Constant);
 }
-TNiagaraExprPtr FNiagaraCompiler::GetExternalConstant(const FNiagaraVariableInfo& Constant, const UNiagaraDataObject* Default)
-{
-	UNiagaraDataObject* DefaultDupe = Default ? CastChecked<UNiagaraDataObject>(StaticDuplicateObject(Default, Script, NAME_None, ~RF_Transient)) : nullptr;
-	SetOrAddConstant(false, Constant, DefaultDupe);
-	return Expression_GetExternalConstant(Constant);
-}
 TNiagaraExprPtr FNiagaraCompiler::GetInternalConstant(const FNiagaraVariableInfo& Constant, float Default)
 {
 	SetOrAddConstant(true, Constant, Default);
@@ -419,12 +403,6 @@ TNiagaraExprPtr FNiagaraCompiler::GetInternalConstant(const FNiagaraVariableInfo
 	SetOrAddConstant(true, Constant, Default);
 	return Expression_GetInternalConstant(Constant);
 }
-TNiagaraExprPtr FNiagaraCompiler::GetInternalConstant(const FNiagaraVariableInfo& Constant, const UNiagaraDataObject* Default)
-{
-	UNiagaraDataObject* DefaultDupe = Default ? CastChecked<UNiagaraDataObject>(StaticDuplicateObject(Default, Script, NAME_None, ~RF_Transient)) : nullptr;
-	SetOrAddConstant(true, Constant, DefaultDupe);
-	return Expression_GetInternalConstant(Constant);
-}
 
 bool FNiagaraCompiler::MergeFunctionIntoMainGraph(UNiagaraNodeFunctionCall* FunctionNode, TArray<UNiagaraScript*>& FunctionStack)
 {
diff --git a/Engine/Source/Editor/NiagaraEditor/Private/NiagaraCompiler.h b/Engine/Source/Editor/NiagaraEditor/Private/NiagaraCompiler.h
--- a/Engine/Source/Editor/NiagaraEditor/Private/NiagaraCompiler.h
+++ b/Engine/Source/Editor/NiagaraEditor/Private/NiagaraCompiler.h
@@ -7,6 +7,11 @@
 #include "NiagaraScriptConstantData.h"
 #include "NiagaraScript.h"
 
+class UNiagaraGraph;
+class UNiagaraScriptSource;
+class UNiagaraScript;
+class UNiagaraNodeFunctionCall;
+
 /** Base class for Niagara compilers. Children of this will include a compiler for the VectorVM and for Compute shaders. Possibly others. */
 class NIAGARAEDITOR_API FNiagaraCompiler : public INiagaraCompiler
 {
@@ -78,11 +83,9 @@ public:
 	virtual TNiagaraExprPtr GetExternalConstant(const FNiagaraVariableInfo& Constant, float Default)override;
 	virtual TNiagaraExprPtr GetExternalConstant(const FNiagaraVariableInfo& Constant, const FVector4& Default)override;
 	virtual TNiagaraExprPtr GetExternalConstant(const FNiagaraVariableInfo& Constant, const FMatrix& Default)override;
-	virtual TNiagaraExprPtr GetExternalConstant(const FNiagaraVariableInfo& Constant, const UNiagaraDataObject* Default)override;
 	virtual TNiagaraExprPtr GetInternalConstant(const FNiagaraVariableInfo& Constant, float Default)override;
 	virtual TNiagaraExprPtr GetInternalConstant(const FNiagaraVariableInfo& Constant, const FVector4& Default)override;
 	virtual TNiagaraExprPtr GetInternalConstant(const FNiagaraVariableInfo& Constant, const FMatrix& Default)override;
-	virtual TNiagaraExprPtr GetInternalConstant(const FNiagaraVariableInfo& Constant, const UNiagaraDataObject* Default)override;
 
 	virtual bool CheckInputs(FName OpName, TArray<TNiagaraExprPtr>& Inputs)override;
 	virtual bool CheckOutputs(FName OpName, TArray<TNiagaraExprPtr>& Outputs)override;
diff --git a/Engine/Source/Editor/NiagaraEditor/Private/NiagaraCompiler_VectorVM.cpp b/Engine/Source/Editor/NiagaraEditor/Private/NiagaraCompiler_VectorVM.cpp
--- a/Engine/Source/Editor/NiagaraEditor/Private/NiagaraCompiler_VectorVM.cpp
+++ b/Engine/Source/Editor/NiagaraEditor/Private/NiagaraCompiler_VectorVM.cpp
@@ -9,6 +9,12 @@
 #include "ComponentReregisterContext.h"
 #include "NiagaraCompiler_VectorVM.h"
 
+#include "NiagaraNodeOutput.h"
+#include "NiagaraNodeInput.h"
+#include "NiagaraNodeFunctionCall.h"
+#include "NiagaraNodeReadDataSet.h"
+#include "NiagaraNodeWriteDataSet.h"
+
 #define LOCTEXT_NAMESPACE "NiagaraCompiler_VectorVM"
 
 DEFINE_LOG_CATEGORY_STATIC(LogNiagaraCompiler_VectorVM, All, All);
@@ -1036,18 +1042,6 @@ bool FNiagaraCompiler_VectorVM::Select_Internal(TArray<TNiagaraExprPtr>& InputEx
 	return true;
 }
 
-bool FNiagaraCompiler_VectorVM::Sample_Internal(TArray<TNiagaraExprPtr>& InputExpressions, TArray<TNiagaraExprPtr>& OutputExpressions)
-{
-	OutputExpressions.Add(Expression_VMNative(EVectorVMOp::sample, InputExpressions));
-	return true;
-}
-
-bool FNiagaraCompiler_VectorVM::Write_Internal(TArray<TNiagaraExprPtr>& InputExpressions, TArray<TNiagaraExprPtr>& OutputExpressions)
-{
-	OutputExpressions.Add(Expression_VMNative(EVectorVMOp::bufferwrite, InputExpressions));
-	return true;
-}
-
 bool FNiagaraCompiler_VectorVM::EaseIn_Internal(TArray<TNiagaraExprPtr>& InputExpressions, TArray<TNiagaraExprPtr>& OutputExpressions)
 {
 	OutputExpressions.Add(Expression_VMNative(EVectorVMOp::easein, InputExpressions));
diff --git a/Engine/Source/Editor/NiagaraEditor/Private/NiagaraEditor.h b/Engine/Source/Editor/NiagaraEditor/Private/NiagaraEditor.h
--- a/Engine/Source/Editor/NiagaraEditor/Private/NiagaraEditor.h
+++ b/Engine/Source/Editor/NiagaraEditor/Private/NiagaraEditor.h
@@ -5,6 +5,7 @@
 #include "INiagaraEditor.h"
 #include "Toolkits/AssetEditorToolkit.h"
 
+class UNiagaraScriptSource;
 
 /** Viewer/editor for a DataTable */
 class FNiagaraEditor : public INiagaraEditor, public FNotifyHook, public FEditorUndoClient
diff --git a/Engine/Source/Editor/NiagaraEditor/Private/NiagaraEditorPrivatePCH.h b/Engine/Source/Editor/NiagaraEditor/Private/NiagaraEditorPrivatePCH.h
--- a/Engine/Source/Editor/NiagaraEditor/Private/NiagaraEditorPrivatePCH.h
+++ b/Engine/Source/Editor/NiagaraEditor/Private/NiagaraEditorPrivatePCH.h
@@ -3,11 +3,13 @@
 #pragma once
 
 #include "UnrealEd.h"
-
+#include "NiagaraEditorCommon.h"
 #include "NiagaraEditorModule.h"
 #include "NiagaraEditor.h"
 #include "NiagaraCompiler.h"
 #include "NiagaraEffectEditor.h"
+#include "NiagaraGraph.h"
+#include "NiagaraScriptSource.h"
 
 #include "GraphEditor.h"
 
diff --git a/Engine/Source/Editor/NiagaraEditor/Private/NiagaraEffectEditor.cpp b/Engine/Source/Editor/NiagaraEditor/Private/NiagaraEffectEditor.cpp
--- a/Engine/Source/Editor/NiagaraEditor/Private/NiagaraEffectEditor.cpp
+++ b/Engine/Source/Editor/NiagaraEditor/Private/NiagaraEffectEditor.cpp
@@ -440,23 +440,6 @@ FReply FNiagaraEffectEditor::OnDeleteEmitterClicked(TSharedPtr<FNiagaraSimulatio
 
 FReply FNiagaraEffectEditor::OnEmitterSelected(TSharedPtr<FNiagaraSimulation> SelectedItem, ESelectInfo::Type SelType)
 {
- 	if (SelectedItem.Get() != nullptr)
-	{
-		if (UNiagaraEmitterProperties* PinnedProps = SelectedItem->GetProperties().Get())
-		{
-			if (PinnedProps->UpdateScriptProps.ExternalConstants.GetNumDataObjectConstants() > 0)
-			{
-				FNiagaraVariableInfo VarInfo;
-				UNiagaraDataObject* DataObj;
-				PinnedProps->UpdateScriptProps.ExternalConstants.GetDataObjectConstant(0, DataObj, VarInfo);
-
-				if (UNiagaraCurveDataObject* CurvObj = Cast<UNiagaraCurveDataObject>(DataObj))
-				{
-					TimeLine.Get()->SetCurve(CurvObj->GetCurveObject());
-				}
-			}
-		}
-	}
 	return FReply::Handled();
 }
 
diff --git a/Engine/Source/Editor/NiagaraEditor/Private/NiagaraEffectFactory.cpp b/Engine/Source/Editor/NiagaraEditor/Private/NiagaraEffectFactory.cpp
new file mode 100644
--- /dev/null
+++ b/Engine/Source/Editor/NiagaraEditor/Private/NiagaraEffectFactory.cpp
@@ -0,0 +1,39 @@
+// Copyright 1998-2016 Epic Games, Inc. All Rights Reserved.
+
+#include "NiagaraEditorPrivatePCH.h"
+#include "NiagaraEffect.h"
+#include "NiagaraEditorModule.h"
+#include "NiagaraEffectFactoryNew.h"
+
+#define LOCTEXT_NAMESPACE "NiagaraEffectFactory"
+
+UNiagaraEffectFactoryNew::UNiagaraEffectFactoryNew(const FObjectInitializer& ObjectInitializer)
+: Super(ObjectInitializer)
+{
+
+	SupportedClass = UNiagaraEffect::StaticClass();
+	bCreateNew = false;
+	bEditAfterNew = true;
+
+	// Check ini to see if we should enable creation
+	bool bEnableNiagara = false;
+	GConfig->GetBool(TEXT("Niagara"), TEXT("EnableNiagara"), bEnableNiagara, GEngineIni);
+	bCreateNew = bEnableNiagara;
+}
+
+UObject* UNiagaraEffectFactoryNew::FactoryCreateNew(UClass* Class, UObject* InParent, FName Name, EObjectFlags Flags, UObject* Context, FFeedbackContext* Warn)
+{
+	check(Class->IsChildOf(UNiagaraEffect::StaticClass()));
+
+	// First allocate runtime script 
+	UNiagaraEffect* NewEffect = NewObject<UNiagaraEffect>(InParent, Class, Name, Flags);
+	if (NewEffect != NULL)
+	{
+		// Then allocate editor-only objects
+		//FNiagaraEditorModule& NiagaraEditorModule = FModuleManager::Get().LoadModuleChecked<FNiagaraEditorModule>(TEXT("NiagaraEditor"));
+	}
+
+	return NewEffect;
+}
+
+#undef LOCTEXT_NAMESPACE
\ No newline at end of file
diff --git a/Engine/Source/Editor/NiagaraEditor/Private/NiagaraEmitterPropertiesDetailsCustomization.cpp b/Engine/Source/Editor/NiagaraEditor/Private/NiagaraEmitterPropertiesDetailsCustomization.cpp
--- a/Engine/Source/Editor/NiagaraEditor/Private/NiagaraEmitterPropertiesDetailsCustomization.cpp
+++ b/Engine/Source/Editor/NiagaraEditor/Private/NiagaraEmitterPropertiesDetailsCustomization.cpp
@@ -73,11 +73,6 @@ void FNiagaraEmitterPropertiesDetails::BuildScriptProperties(TSharedRef<IPropert
 	TSharedRef<FDetailArrayBuilder> MatrixConstantsBuilder = MakeShareable(new FDetailArrayBuilder(MatrixConstants.ToSharedRef(), bGenerateHeader, bDisplayResetToDefault));
 	MatrixConstantsBuilder->OnGenerateArrayElementWidget(FOnGenerateArrayElementWidget::CreateSP(this, &FNiagaraEmitterPropertiesDetails::OnGenerateMatrixConstantEntry));
 	ScriptCategory.AddCustomBuilder(MatrixConstantsBuilder);
-	//DataObject Constants.
-	TSharedPtr<IPropertyHandle> DataObjectConstants = Constants->GetChildHandle(GET_MEMBER_NAME_CHECKED(FNiagaraConstants, DataObjectConstants));
-	TSharedRef<FDetailArrayBuilder> DataObjectConstantsBuilder = MakeShareable(new FDetailArrayBuilder(DataObjectConstants.ToSharedRef(), bGenerateHeader, bDisplayResetToDefault));
-	DataObjectConstantsBuilder->OnGenerateArrayElementWidget(FOnGenerateArrayElementWidget::CreateSP(this, &FNiagaraEmitterPropertiesDetails::OnGenerateDataObjectConstantEntry));
-	ScriptCategory.AddCustomBuilder(DataObjectConstantsBuilder);
 
 	//Events
 	bGenerateHeader = true;
@@ -136,20 +131,6 @@ void FNiagaraEmitterPropertiesDetails::OnGenerateMatrixConstantEntry(TSharedRef<
 	}
 }
 
-void FNiagaraEmitterPropertiesDetails::OnGenerateDataObjectConstantEntry(TSharedRef<IPropertyHandle> ElementProperty, int32 ElementIndex, IDetailChildrenBuilder& ChildrenBuilder)
-{
-	TSharedPtr<IPropertyHandle> ValueProperty = ElementProperty->GetChildHandle(GET_MEMBER_NAME_CHECKED(FNiagaraConstants_DataObject, Value));
-	TSharedPtr<IPropertyHandle> NameProperty = ElementProperty->GetChildHandle(GET_MEMBER_NAME_CHECKED(FNiagaraConstants_DataObject, Name));
-
-	FName DisplayName;
-	NameProperty->GetValue(DisplayName);
-	//Don't display system constants
-	if (UNiagaraComponent::GetSystemConstants().Find(FNiagaraVariableInfo(DisplayName, ENiagaraDataType::Curve)) == INDEX_NONE)
-	{
-		ChildrenBuilder.AddChildProperty(ValueProperty.ToSharedRef()).DisplayName(FText::FromName(DisplayName));
-	}
-}
-
 void FNiagaraEmitterPropertiesDetails::OnGenerateEventGeneratorEntry(TSharedRef<IPropertyHandle> ElementProperty, int32 ElementIndex, IDetailChildrenBuilder& ChildrenBuilder)
 {
 	TSharedPtr<IPropertyHandle> IDProperty = ElementProperty->GetChildHandle(GET_MEMBER_NAME_CHECKED(FNiagaraEventGeneratorProperties, ID));
diff --git a/Engine/Source/Editor/NiagaraEditor/Private/NiagaraEmitterPropertiesDetailsCustomization.h b/Engine/Source/Editor/NiagaraEditor/Private/NiagaraEmitterPropertiesDetailsCustomization.h
--- a/Engine/Source/Editor/NiagaraEditor/Private/NiagaraEmitterPropertiesDetailsCustomization.h
+++ b/Engine/Source/Editor/NiagaraEditor/Private/NiagaraEmitterPropertiesDetailsCustomization.h
@@ -21,7 +21,6 @@ public:
 	void OnGenerateScalarConstantEntry(TSharedRef<IPropertyHandle> ElementProperty, int32 ElementIndex, IDetailChildrenBuilder& ChildrenBuilder);
 	void OnGenerateVectorConstantEntry(TSharedRef<IPropertyHandle> ElementProperty, int32 ElementIndex, IDetailChildrenBuilder& ChildrenBuilder);
 	void OnGenerateMatrixConstantEntry(TSharedRef<IPropertyHandle> ElementProperty, int32 ElementIndex, IDetailChildrenBuilder& ChildrenBuilder);
-	void OnGenerateDataObjectConstantEntry(TSharedRef<IPropertyHandle> ElementProperty, int32 ElementIndex, IDetailChildrenBuilder& ChildrenBuilder);
 	void OnGenerateEventGeneratorEntry(TSharedRef<IPropertyHandle> ElementProperty, int32 ElementIndex, IDetailChildrenBuilder& ChildrenBuilder);
 	void OnGenerateEventReceiverEntry(TSharedRef<IPropertyHandle> ElementProperty, int32 ElementIndex, IDetailChildrenBuilder& ChildrenBuilder);
 
diff --git a/Engine/Source/Editor/NiagaraEditor/Private/NiagaraNode.cpp b/Engine/Source/Editor/NiagaraEditor/Private/NiagaraNode.cpp
new file mode 100644
--- /dev/null
+++ b/Engine/Source/Editor/NiagaraEditor/Private/NiagaraNode.cpp
@@ -0,0 +1,150 @@
+// Copyright 1998-2016 Epic Games, Inc. All Rights Reserved.
+
+#include "NiagaraEditorPrivatePCH.h"
+#include "NiagaraNode.h"
+#include "EdGraphSchema_Niagara.h"
+
+UNiagaraNode::UNiagaraNode(const FObjectInitializer& ObjectInitializer)
+	: Super(ObjectInitializer)
+{
+}
+
+void UNiagaraNode::ReallocatePins()
+{
+	Modify();
+
+	// Move the existing pins to a saved array
+	TArray<UEdGraphPin*> OldPins(Pins);
+	Pins.Empty();
+
+	// Recreate the new pins
+	AllocateDefaultPins();
+
+	// Copy the old pin data and remove it.
+	for (int32 OldPinIndex = 0; OldPinIndex < OldPins.Num(); ++OldPinIndex)
+	{
+		UEdGraphPin* OldPin = OldPins[OldPinIndex];
+		if (UEdGraphPin** MatchingNewPin = Pins.FindByPredicate([&](UEdGraphPin* Pin){ return Pin->Direction == OldPin->Direction && Pin->PinName == OldPin->PinName; }))
+		{
+			if (*MatchingNewPin && OldPin)
+				(*MatchingNewPin)->CopyPersistentDataFromOldPin(*OldPin);
+		}
+		OldPin->Modify();
+		OldPin->BreakAllPinLinks();
+	}
+	OldPins.Empty();
+
+	GetGraph()->NotifyGraphChanged();
+}
+
+void UNiagaraNode::AutowireNewNode(UEdGraphPin* FromPin)
+{
+	Super::AutowireNewNode(FromPin);
+
+	if (FromPin != nullptr)
+	{
+		const UEdGraphSchema_Niagara* Schema = CastChecked<UEdGraphSchema_Niagara>(GetSchema());
+		check(Schema);
+		
+		ENiagaraDataType FromType = Schema->GetPinDataType(FromPin);
+
+		//Find first of this nodes pins with the right type and direction.
+		UEdGraphPin* FirstPinOfSameType = NULL;
+		EEdGraphPinDirection DesiredDirection = FromPin->Direction == EGPD_Output ? EGPD_Input : EGPD_Output;
+		for (UEdGraphPin* Pin : Pins)
+		{
+			ENiagaraDataType ToType = Schema->GetPinDataType(Pin);
+			if (FromType == ToType && Pin->Direction == DesiredDirection)
+			{
+				FirstPinOfSameType = Pin;
+				break;
+			}			
+		}
+
+		if (FirstPinOfSameType && GetSchema()->TryCreateConnection(FromPin, FirstPinOfSameType))
+		{
+			FromPin->GetOwningNode()->NodeConnectionListChanged();
+		}
+	}
+}
+
+const UNiagaraGraph* UNiagaraNode::GetNiagaraGraph()const
+{
+	return CastChecked<UNiagaraGraph>(GetGraph());
+}
+
+UNiagaraGraph* UNiagaraNode::GetNiagaraGraph()
+{
+	return CastChecked<UNiagaraGraph>(GetGraph());
+}
+
+UNiagaraScriptSource* UNiagaraNode::GetSource()const
+{
+	return GetNiagaraGraph()->GetSource();
+}
+
+UEdGraphPin* UNiagaraNode::GetInputPin(int32 InputIndex) const
+{
+	for (int32 PinIndex = 0, FoundInputs = 0; PinIndex < Pins.Num(); PinIndex++)
+	{
+		if (Pins[PinIndex]->Direction == EGPD_Input)
+		{
+			if (InputIndex == FoundInputs)
+			{
+				return Pins[PinIndex];
+			}
+			else
+			{
+				FoundInputs++;
+			}
+		}
+	}
+
+	return NULL;
+}
+
+void UNiagaraNode::GetInputPins(TArray<class UEdGraphPin*>& OutInputPins) const
+{
+	OutInputPins.Empty();
+
+	for (int32 PinIndex = 0; PinIndex < Pins.Num(); PinIndex++)
+	{
+		if (Pins[PinIndex]->Direction == EGPD_Input)
+		{
+			OutInputPins.Add(Pins[PinIndex]);
+		}
+	}
+}
+
+UEdGraphPin* UNiagaraNode::GetOutputPin(int32 OutputIndex) const
+{
+	for (int32 PinIndex = 0, FoundOutputs = 0; PinIndex < Pins.Num(); PinIndex++)
+	{
+		if (Pins[PinIndex]->Direction == EGPD_Output)
+		{
+			if (OutputIndex == FoundOutputs)
+			{
+				return Pins[PinIndex];
+			}
+			else
+			{
+				FoundOutputs++;
+			}
+		}
+	}
+
+	return NULL;
+}
+
+void UNiagaraNode::GetOutputPins(TArray<class UEdGraphPin*>& OutOutputPins) const
+{
+	OutOutputPins.Empty();
+
+	for (int32 PinIndex = 0; PinIndex < Pins.Num(); PinIndex++)
+	{
+		if (Pins[PinIndex]->Direction == EGPD_Output)
+		{
+			OutOutputPins.Add(Pins[PinIndex]);
+		}
+	}
+}
\ No newline at end of file
diff --git a/Engine/Source/Editor/NiagaraEditor/Private/NiagaraNodeConstant.cpp b/Engine/Source/Editor/NiagaraEditor/Private/NiagaraNodeConstant.cpp
new file mode 100644
--- /dev/null
+++ b/Engine/Source/Editor/NiagaraEditor/Private/NiagaraNodeConstant.cpp
@@ -0,0 +1,4 @@
+// Copyright 1998-2016 Epic Games, Inc. All Rights Reserved.
+
+#include "NiagaraEditorPrivatePCH.h"
+#include "INiagaraCompiler.h"
diff --git a/Engine/Source/Editor/NiagaraEditor/Private/NiagaraNodeGetAttr.cpp b/Engine/Source/Editor/NiagaraEditor/Private/NiagaraNodeGetAttr.cpp
new file mode 100644
--- /dev/null
+++ b/Engine/Source/Editor/NiagaraEditor/Private/NiagaraNodeGetAttr.cpp
@@ -0,0 +1,4 @@
+// Copyright 1998-2016 Epic Games, Inc. All Rights Reserved.
+
+#include "NiagaraEditorPrivatePCH.h"
+#include "INiagaraCompiler.h"
diff --git a/Engine/Source/Editor/NiagaraEditor/Private/NiagaraNodeOutputUpdate.cpp b/Engine/Source/Editor/NiagaraEditor/Private/NiagaraNodeOutputUpdate.cpp
new file mode 100644
--- /dev/null
+++ b/Engine/Source/Editor/NiagaraEditor/Private/NiagaraNodeOutputUpdate.cpp
@@ -0,0 +1,6 @@
+// Copyright 1998-2016 Epic Games, Inc. All Rights Reserved.
+
+#include "NiagaraEditorPrivatePCH.h"
+#include "BlueprintGraphDefinitions.h"
+#include "GraphEditorSettings.h"
+
diff --git a/Engine/Source/Editor/NiagaraEditor/Private/NiagaraNodeReadDataSet.cpp b/Engine/Source/Editor/NiagaraEditor/Private/NiagaraNodeReadDataSet.cpp
new file mode 100644
--- /dev/null
+++ b/Engine/Source/Editor/NiagaraEditor/Private/NiagaraNodeReadDataSet.cpp
@@ -0,0 +1,94 @@
+// Copyright 1998-2016 Epic Games, Inc. All Rights Reserved.
+
+#include "NiagaraEditorPrivatePCH.h"
+#include "BlueprintGraphDefinitions.h"
+#include "GraphEditorSettings.h"
+
+#include "NiagaraNodeReadDataSet.h"
+#include "EdGraphSchema_Niagara.h"
+
+#define LOCTEXT_NAMESPACE "NiagaraNodeDataSetRead"
+
+UNiagaraNodeReadDataSet::UNiagaraNodeReadDataSet(const FObjectInitializer& ObjectInitializer)
+: Super(ObjectInitializer)
+{
+}
+
+void UNiagaraNodeReadDataSet::PostEditChangeProperty(struct FPropertyChangedEvent& PropertyChangedEvent)
+{
+	if (PropertyChangedEvent.Property != nullptr)
+	{
+		ReallocatePins();
+	}
+	Super::PostEditChangeProperty(PropertyChangedEvent);
+}
+
+void UNiagaraNodeReadDataSet::AllocateDefaultPins()
+{
+	const UEdGraphSchema_Niagara* Schema = GetDefault<UEdGraphSchema_Niagara>();
+
+	if (DataSet.Type == ENiagaraDataSetType::Event)
+	{
+		//Probably need this for all data set types tbh!
+		UEdGraphPin* Pin = CreatePin(EGPD_Output, Schema->PC_Vector, TEXT(""), NULL, false, false, TEXT("Valid"));//TODO - CHANGE TO INT / BOOL / SCALAR
+		Pin->bDefaultValueIsIgnored = true;
+	}
+
+	for (const FNiagaraVariableInfo& Var : Variables)
+	{
+		switch (Var.Type)
+		{
+		case ENiagaraDataType::Scalar:
+		{
+			CreatePin(EGPD_Output, Schema->PC_Float, TEXT(""), NULL, false, false, Var.Name.ToString());
+		}
+			break;
+		case ENiagaraDataType::Vector:
+		{
+			CreatePin(EGPD_Output, Schema->PC_Vector, TEXT(""), NULL, false, false, Var.Name.ToString());
+		}
+			break;
+		case ENiagaraDataType::Matrix:
+		{
+			CreatePin(EGPD_Output, Schema->PC_Matrix, TEXT(""), NULL, false, false, Var.Name.ToString());
+		}
+			break;
+		};
+	}
+}
+
+FText UNiagaraNodeReadDataSet::GetNodeTitle(ENodeTitleType::Type TitleType) const
+{
+	return FText::Format(LOCTEXT("NiagaraDataSetReadFormat", "{0} Read"), FText::FromName(DataSet.Name));
+}
+
+FLinearColor UNiagaraNodeReadDataSet::GetNodeTitleColor() const
+{
+	check(DataSet.Type == ENiagaraDataSetType::Event);//Implement other datasets
+	return CastChecked<UEdGraphSchema_Niagara>(GetSchema())->NodeTitleColor_Event;
+}
+
+void UNiagaraNodeReadDataSet::Compile(class INiagaraCompiler* Compiler, TArray<FNiagaraNodeResult>& Outputs)
+{
+	if(DataSet.Type == ENiagaraDataSetType::Event)
+	{
+		TNiagaraExprPtr IndexExpr = Compiler->ConsumeSharedDataIndex(DataSet, true);//Where do we decide wrapping?
+		check(Variables.Num() == Pins.Num() - 1);//We should have a pin for each variable + 1 for the Valid pin.
+		Outputs.Add(FNiagaraNodeResult(Compiler->SharedDataIndexIsValid(DataSet, IndexExpr), Pins[0]));
+		for (int32 i = 0; i < Variables.Num(); ++i)
+		{
+			check(Variables[i].Type == ENiagaraDataType::Vector);//Only support vectors currently.
+			Outputs.Add(FNiagaraNodeResult(Compiler->SharedDataRead(DataSet, Variables[i], IndexExpr), Pins[i+1]));
+		}
+	}
+	else
+	{
+		Compiler->Error(LOCTEXT("UnImplDataSetTypeError", "This data set type hasn't been implemented."), this, nullptr);
+	}
+	
+}
+
+#undef LOCTEXT_NAMESPACE
+
+
+
diff --git a/Engine/Source/Editor/NiagaraEditor/Private/NiagaraNodeWriteDataSet.cpp b/Engine/Source/Editor/NiagaraEditor/Private/NiagaraNodeWriteDataSet.cpp
new file mode 100644
--- /dev/null
+++ b/Engine/Source/Editor/NiagaraEditor/Private/NiagaraNodeWriteDataSet.cpp
@@ -0,0 +1,133 @@
+// Copyright 1998-2016 Epic Games, Inc. All Rights Reserved.
+
+#include "NiagaraEditorPrivatePCH.h"
+#include "BlueprintGraphDefinitions.h"
+#include "GraphEditorSettings.h"
+
+#include "NiagaraNodeWriteDataSet.h"
+#include "EdGraphSchema_Niagara.h"
+
+#define LOCTEXT_NAMESPACE "NiagaraNodeWriteDataSet"
+
+UNiagaraNodeWriteDataSet::UNiagaraNodeWriteDataSet(const FObjectInitializer& ObjectInitializer)
+: Super(ObjectInitializer)
+{
+}
+
+void UNiagaraNodeWriteDataSet::PostEditChangeProperty(struct FPropertyChangedEvent& PropertyChangedEvent)
+{
+	if (PropertyChangedEvent.Property != nullptr)
+	{
+		ReallocatePins();
+	}
+	Super::PostEditChangeProperty(PropertyChangedEvent);
+}
+
+void UNiagaraNodeWriteDataSet::AllocateDefaultPins()
+{
+	const UEdGraphSchema_Niagara* Schema = GetDefault<UEdGraphSchema_Niagara>();
+
+	if (DataSet.Type == ENiagaraDataSetType::Event)
+	{
+		UEdGraphPin* Pin = CreatePin(EGPD_Input, Schema->PC_Vector, TEXT(""), NULL, false, false, TEXT("Valid"));//TODO - CHANGE TO INT / BOOL / SCALAR
+		Pin->bDefaultValueIsIgnored = true;
+	}
+	
+	for (const FNiagaraVariableInfo& Var : Variables)
+	{
+		switch (Var.Type)
+		{
+		case ENiagaraDataType::Scalar:
+		{
+			CreatePin(EGPD_Input, Schema->PC_Float, TEXT(""), NULL, false, false, Var.Name.ToString());
+		}
+			break;
+		case ENiagaraDataType::Vector:
+		{
+			CreatePin(EGPD_Input, Schema->PC_Vector, TEXT(""), NULL, false, false, Var.Name.ToString());
+		}
+			break;
+		case ENiagaraDataType::Matrix:
+		{
+			CreatePin(EGPD_Input, Schema->PC_Matrix, TEXT(""), NULL, false, false, Var.Name.ToString());
+		}
+			break;
+		};
+	}
+}
+
+FText UNiagaraNodeWriteDataSet::GetNodeTitle(ENodeTitleType::Type TitleType) const
+{
+	return FText::Format(LOCTEXT("NiagaraDataSetWriteFormat", "{0} Write"), FText::FromName(DataSet.Name));
+}
+
+FLinearColor UNiagaraNodeWriteDataSet::GetNodeTitleColor() const
+{
+	check(DataSet.Type == ENiagaraDataSetType::Event);//Implement other datasets
+	return CastChecked<UEdGraphSchema_Niagara>(GetSchema())->NodeTitleColor_Event;
+}
+
+void UNiagaraNodeWriteDataSet::Compile(class INiagaraCompiler* Compiler, TArray<FNiagaraNodeResult>& Outputs)
+{
+	bool bError = false;
+
+	if (DataSet.Type == ENiagaraDataSetType::Event)
+	{
+		//Compile the Emit pin.
+		TNiagaraExprPtr EmitExpression = Compiler->CompilePin(Pins[0]);
+
+		if (EmitExpression.IsValid())
+		{
+			//Test the Valid pin result against 0. Maybe just require a direct connection of 0 or 0xFFFFFFFF?
+			FNiagaraVariableInfo Zero(TEXT("0.0, 0.0, 0.0, 0.0"), ENiagaraDataType::Vector);
+			TNiagaraExprPtr ZeroContantExpression = Compiler->GetInternalConstant(Zero, FVector4(0.0f, 0.0f, 0.0f, 0.0f));
+			TArray<TNiagaraExprPtr> ConditonInputs;
+			ConditonInputs.Add(EmitExpression);
+			ConditonInputs.Add(ZeroContantExpression);
+			check(ZeroContantExpression.IsValid());
+			TArray<TNiagaraExprPtr> ConditionOpOutputs;
+			INiagaraCompiler::GreaterThan(Compiler, ConditonInputs, ConditionOpOutputs);
+			TNiagaraExprPtr ValidExpr = ConditionOpOutputs[0];
+
+			TArray<TNiagaraExprPtr> InputExpressions;
+			for (int32 i = 0; i < Variables.Num(); ++i)
+			{
+				const FNiagaraVariableInfo& Var = Variables[i];
+				UEdGraphPin* Pin = Pins[i + 1];//Pin[0] is emit
+				check(Pin->Direction == EGPD_Input);
+
+				TNiagaraExprPtr Result = Compiler->CompilePin(Pin);
+				if (!Result.IsValid())
+				{
+					bError = true;
+					Compiler->Error(FText::Format(LOCTEXT("DataSetWriteErrorFormat", "Error writing variable {0} to dataset {1}"), FText::FromName(DataSet.Name), Pin->PinFriendlyName), this, Pin);
+				}
+
+				InputExpressions.Add(Result);
+			}
+
+			//Gets the index to write to. 
+			TNiagaraExprPtr IndexExpression = Compiler->AcquireSharedDataIndex(DataSet, true, ValidExpr);
+
+			if (!bError)
+			{
+				check(Variables.Num() == InputExpressions.Num());
+				for (int32 i = 0; i < Variables.Num(); ++i)
+				{
+					Outputs.Add(FNiagaraNodeResult(Compiler->SharedDataWrite(DataSet, Variables[i], IndexExpression, InputExpressions[i]), Pins[i + 1]));//Pin[0] is Emit
+				}
+			}
+		}
+	}
+	else
+	{
+		check(false);//IMPLEMENT OTHER DATA SETS.
+	}
+}
+
+#undef LOCTEXT_NAMESPACE
+
+
+
+
+
diff --git a/Engine/Source/Editor/NiagaraEditor/Private/NiagaraScriptFactory.cpp b/Engine/Source/Editor/NiagaraEditor/Private/NiagaraScriptFactory.cpp
new file mode 100644
--- /dev/null
+++ b/Engine/Source/Editor/NiagaraEditor/Private/NiagaraScriptFactory.cpp
@@ -0,0 +1,65 @@
+// Copyright 1998-2016 Epic Games, Inc. All Rights Reserved.
+
+#include "NiagaraEditorPrivatePCH.h"
+#include "NiagaraScript.h"
+#include "NiagaraEditorModule.h"
+#include "NiagaraScriptFactoryNew.h"
+#include "NiagaraNodeOutput.h"
+
+#define LOCTEXT_NAMESPACE "NiagaraScriptFactory"
+
+UNiagaraScriptFactoryNew::UNiagaraScriptFactoryNew(const FObjectInitializer& ObjectInitializer)
+	: Super(ObjectInitializer)
+{
+
+	SupportedClass = UNiagaraScript::StaticClass();
+	bCreateNew = false;
+	bEditAfterNew = true;
+
+	// Check ini to see if we should enable creation
+	bool bEnableNiagara = false;
+	GConfig->GetBool(TEXT("Niagara"), TEXT("EnableNiagara"), bEnableNiagara, GEngineIni);
+	bCreateNew = bEnableNiagara;
+}
+
+UObject* UNiagaraScriptFactoryNew::FactoryCreateNew(UClass* Class,UObject* InParent,FName Name,EObjectFlags Flags,UObject* Context,FFeedbackContext* Warn)
+{
+	check(Class->IsChildOf(UNiagaraScript::StaticClass()));
+
+	// First allocate runtime script 
+	UNiagaraScript* NewScript = NewObject<UNiagaraScript>(InParent, Class, Name, Flags);
+	if(NewScript != NULL)
+	{
+		// Then allocate editor-only 'source' object
+		UNiagaraScriptSource* Source = NewObject<UNiagaraScriptSource>(NewScript, NAME_None, RF_Transactional);
+		if(Source)
+		{
+			// Create 'update graph'
+			UNiagaraGraph* CreatedGraph = NewObject<UNiagaraGraph>(Source, NAME_None, RF_Transactional);
+			Source->NodeGraph = CreatedGraph;
+
+			FGraphNodeCreator<UNiagaraNodeOutput> OutputNodeCreator(*CreatedGraph);
+			UNiagaraNodeOutput* OutputNode = OutputNodeCreator.CreateNode();
+
+			//Automatically add attributes that the current render modules rely on. Not sure what we'll want to do here in the end.
+			OutputNode->Outputs.Add(FNiagaraVariableInfo(FName(TEXT("Position")), ENiagaraDataType::Vector));
+			OutputNode->Outputs.Add(FNiagaraVariableInfo(FName(TEXT("Velocity")), ENiagaraDataType::Vector));
+			OutputNode->Outputs.Add(FNiagaraVariableInfo(FName(TEXT("Acceleration")), ENiagaraDataType::Vector));
+			OutputNode->Outputs.Add(FNiagaraVariableInfo(FName(TEXT("Rotation")), ENiagaraDataType::Vector));
+			OutputNode->Outputs.Add(FNiagaraVariableInfo(FName(TEXT("Color")), ENiagaraDataType::Vector));
+			OutputNode->Outputs.Add(FNiagaraVariableInfo(FName(TEXT("Age")), ENiagaraDataType::Vector));
+
+			OutputNodeCreator.Finalize();
+
+			// Set pointer in script to source
+			NewScript->Source = Source;
+
+			FNiagaraEditorModule& NiagaraEditorModule = FModuleManager::Get().LoadModuleChecked<FNiagaraEditorModule>(TEXT("NiagaraEditor"));
+			NiagaraEditorModule.CompileScript(NewScript);
+		}
+	}
+
+	return NewScript;
+}
+
+#undef LOCTEXT_NAMESPACE
\ No newline at end of file
diff --git a/Engine/Source/Editor/NiagaraEditor/Private/NiagaraScriptSource.cpp b/Engine/Source/Editor/NiagaraEditor/Private/NiagaraScriptSource.cpp
new file mode 100644
--- /dev/null
+++ b/Engine/Source/Editor/NiagaraEditor/Private/NiagaraScriptSource.cpp
@@ -0,0 +1,237 @@
+// Copyright 1998-2016 Epic Games, Inc. All Rights Reserved.
+
+#include "NiagaraEditorPrivatePCH.h"
+#include "NiagaraScript.h"
+#include "NiagaraConstants.h"
+#include "NiagaraEditorModule.h"
+#include "ComponentReregisterContext.h"
+#include "NiagaraComponent.h"
+#include "NiagaraSimulation.h"
+#include "NiagaraEffect.h"
+#include "NiagaraNodeOutput.h"
+#include "NiagaraNodeInput.h"
+#include "NiagaraNodeWriteDataSet.h"
+#include "NiagaraNodeReadDataSet.h"
+
+#include "EdgraphSchema_Niagara.h"
+
+//////////////////////////////////////////////////////////////////////////
+// NiagaraGraph
+
+UNiagaraGraph::UNiagaraGraph(const FObjectInitializer& ObjectInitializer)
+	: Super(ObjectInitializer)
+{
+
+	Schema = UEdGraphSchema_Niagara::StaticClass();
+}
+
+void UNiagaraGraph::PostEditChangeProperty(FPropertyChangedEvent& PropertyChangedEvent)
+{
+	NotifyGraphChanged();
+}
+
+class UNiagaraScriptSource* UNiagaraGraph::GetSource() const
+{
+	return CastChecked<UNiagaraScriptSource>(GetOuter());
+}
+
+UNiagaraNodeOutput* UNiagaraGraph::FindOutputNode() const
+{
+	for (UEdGraphNode* Node : Nodes)
+	{
+		if (UNiagaraNodeOutput* OutNode = Cast<UNiagaraNodeOutput>(Node))
+		{
+			return OutNode;
+		}
+	}
+	check(0);
+	return NULL;
+}
+
+void UNiagaraGraph::FindInputNodes(TArray<class UNiagaraNodeInput*>& OutInputNodes) const
+{
+	for (UEdGraphNode* Node : Nodes)
+	{
+		if (UNiagaraNodeInput* InNode = Cast<UNiagaraNodeInput>(Node))
+		{
+			OutInputNodes.Add(InNode);
+		}
+	}
+}
+
+void UNiagaraGraph::FindReadDataSetNodes(TArray<class UNiagaraNodeReadDataSet*>& OutReadNodes) const
+{
+	for (UEdGraphNode* Node : Nodes)
+	{
+		if (UNiagaraNodeReadDataSet* InNode = Cast<UNiagaraNodeReadDataSet>(Node))
+		{
+			OutReadNodes.Add(InNode);
+		}
+	}
+}
+
+void UNiagaraGraph::FindWriteDataSetNodes(TArray<class UNiagaraNodeWriteDataSet*>& OutWriteNodes) const
+{
+	for (UEdGraphNode* Node : Nodes)
+	{
+		if (UNiagaraNodeWriteDataSet* InNode = Cast<UNiagaraNodeWriteDataSet>(Node))
+		{
+			OutWriteNodes.Add(InNode);
+		}
+	}
+}
+
+int32 UNiagaraGraph::GetAttributeIndex(const FNiagaraVariableInfo& Attr)const
+{
+	const UNiagaraNodeOutput* OutNode = FindOutputNode();
+	check(OutNode);
+	for (int32 i = 0; i < OutNode->Outputs.Num(); ++i)
+	{
+		if (OutNode->Outputs[i] == Attr)
+		{
+			return i;
+		}
+	}
+	return INDEX_NONE;
+}
+
+void UNiagaraGraph::GetAttributes(TArray< FNiagaraVariableInfo >& OutAttributes)const
+{
+	const UNiagaraNodeOutput* OutNode = FindOutputNode();
+	check(OutNode);
+
+	for (const FNiagaraVariableInfo& Attr : OutNode->Outputs)
+	{
+		check(!OutAttributes.Find(Attr));
+		OutAttributes.Add(Attr);
+	}
+}
+
+//////////////////////////////////////////////////////////////////////////
+// UNiagraScriptSource
+
+UNiagaraScriptSource::UNiagaraScriptSource(const FObjectInitializer& ObjectInitializer)
+	: Super(ObjectInitializer)
+{
+}
+
+void UNiagaraScriptSource::PostLoad()
+{
+	Super::PostLoad();
+
+#if WITH_EDITOR
+	UNiagaraScript* ScriptOwner = Cast<UNiagaraScript>(GetOuter());
+	if (ScriptOwner && ScriptOwner->ByteCode.Num() == 0)
+	{
+		ScriptOwner->ConditionalPostLoad();
+		//FNiagaraEditorModule& NiagaraEditorModule = FModuleManager::Get().LoadModuleChecked<FNiagaraEditorModule>(TEXT("NiagaraEditor"));
+		//NiagaraEditorModule.CompileScript(ScriptOwner);	
+	}
+	Compile();
+#endif
+}
+
+struct FNiagaraComponentReregisterContext : FComponentReregisterContext
+{
+	FNiagaraEmitterScriptProperties* ScriptProps;
+	UNiagaraEmitterProperties* EmitterProps;
+	FNiagaraComponentReregisterContext(UNiagaraComponent* Comp, FNiagaraEmitterScriptProperties* InScriptProps, UNiagaraEmitterProperties* InEmitterProps)
+		: FComponentReregisterContext(Comp)
+		, ScriptProps(InScriptProps)
+		, EmitterProps(InEmitterProps)
+	{
+
+	}
+	~FNiagaraComponentReregisterContext()
+	{
+		ScriptProps->Init(EmitterProps);
+	}
+};
+
+class FNiagaraScriptCompileContext
+{
+public:
+	/** Initialization constructor. */
+	FNiagaraScriptCompileContext(UNiagaraScript* Script)
+	{
+		// wait until resources are released
+		FlushRenderingCommands();
+
+		// Reregister all components usimg Script.
+		for (TObjectIterator<UNiagaraComponent> ComponentIt; ComponentIt; ++ComponentIt)
+		{
+			UNiagaraComponent* Comp = *ComponentIt;
+			TSharedPtr<FNiagaraEffectInstance> Inst = Comp->GetEffectInstance();
+			if (Inst.IsValid())
+			{
+				TArray<TSharedPtr<FNiagaraSimulation>>& Emitters = Inst->GetEmitters();
+				for (TSharedPtr<FNiagaraSimulation> Sim : Emitters)
+				{
+					if (Sim.IsValid())
+					{
+						if (UNiagaraEmitterProperties* Props = Sim->GetProperties().Get())
+						{
+							if (Props->UpdateScriptProps.Script == Script)
+							{
+								new(ComponentContexts)FNiagaraComponentReregisterContext(Comp, &Props->UpdateScriptProps, Props);
+							}
+							else if (Props->SpawnScriptProps.Script == Script)
+							{
+								new(ComponentContexts)FNiagaraComponentReregisterContext(Comp, &Props->SpawnScriptProps, Props);
+							}
+						}						
+					}
+				}
+			}
+		}
+	}
+	
+private:
+	/** The recreate contexts for the individual components. */
+	TIndirectArray<FNiagaraComponentReregisterContext> ComponentContexts;
+};
+
+void UNiagaraScriptSource::Compile()
+{
+	UNiagaraScript* ScriptOwner = Cast<UNiagaraScript>(GetOuter());
+	
+	FNiagaraScriptCompileContext ScriptCompileContext(ScriptOwner);
+
+	FNiagaraEditorModule& NiagaraEditorModule = FModuleManager::Get().LoadModuleChecked<FNiagaraEditorModule>(TEXT("NiagaraEditor"));
+	NiagaraEditorModule.CompileScript(ScriptOwner);
+
+	FNiagaraConstants& ExternalConsts = ScriptOwner->ConstantData.GetExternalConstants();
+
+	//Build the constant list. 
+	//This is mainly just jumping through some hoops for the custom UI. Should be removed and have the UI just read directly from the constants stored in the UScript.
+	const UEdGraphSchema_Niagara* Schema = CastChecked<UEdGraphSchema_Niagara>(NodeGraph->GetSchema());
+	ExposedVectorConstants.Empty();
+	for (int32 ConstIdx = 0; ConstIdx < ExternalConsts.GetNumVectorConstants(); ConstIdx++)
+	{
+		FNiagaraVariableInfo Info;
+		FVector4 Value;
+		ExternalConsts.GetVectorConstant(ConstIdx, Value, Info);
+		if (Schema->IsSystemConstant(Info))
+		{
+			continue;//System constants are "external" but should not be exposed to the editor.
+		}
+			
+		EditorExposedVectorConstant *Const = new EditorExposedVectorConstant();
+		Const->ConstName = Info.Name;
+		Const->Value = Value;
+		ExposedVectorConstants.Add(MakeShareable(Const));
+	}
+
+}
+
+void UNiagaraScriptSource::GetEmitterAttributes(TArray<FName>& VectorInputs, TArray<FName>& MatrixInputs)
+{
+	for (uint32 i=0; i < NiagaraConstants::NumBuiltinConstants; i++)
+	{
+		VectorInputs.Add(NiagaraConstants::ConstantNames[i]);
+	}
+
+	MatrixInputs.Empty();
+	MatrixInputs.Add(FName(TEXT("Emitter Transform")));
+}
+
diff --git a/Engine/Source/Editor/NiagaraEditor/Private/SNiagaraEffectEditorWidget.h b/Engine/Source/Editor/NiagaraEditor/Private/SNiagaraEffectEditorWidget.h
--- a/Engine/Source/Editor/NiagaraEditor/Private/SNiagaraEffectEditorWidget.h
+++ b/Engine/Source/Editor/NiagaraEditor/Private/SNiagaraEffectEditorWidget.h
@@ -14,10 +14,11 @@
 #include "NiagaraComponent.h"
 #include "NiagaraEffect.h"
 #include "NiagaraSimulation.h"
-#include "Engine/NiagaraEffectRenderer.h"
+#include "NiagaraEffectRenderer.h"
 #include "ComponentReregisterContext.h"
 #include "SNumericEntryBox.h"
 #include "Curves/CurveVector.h"
+#include "NiagaraScriptSourceBase.h"
 
 #define NGED_SECTION_BORDER SNew(SBorder).BorderImage(FEditorStyle::GetBrush("ToolPanel.GroupBorder")).Padding(2.0f).HAlign(HAlign_Left)
 #define NGED_SECTION_LIGHTBORDER SNew(SBorder).BorderImage(FEditorStyle::GetBrush("ToolPanel.LightGroupBorder")).Padding(2.0f).HAlign(HAlign_Left)
@@ -193,31 +194,10 @@ public:
 	void OnCurveSelected(UObject *Asset)
 	{
 		check(ScriptProps);
-		CurCurve = Cast<UCurveVector>(Asset);
-		UNiagaraEmitterProperties* PinnedProps = EmitterProps.Get();
-		if (PinnedProps)
-		{
-			UNiagaraDataObject *DataObj = ScriptProps->ExternalConstants.FindDataObj(ConstantName);
-			if (!DataObj)
-			{
-				//In the details UI this is created via a combo box. Not sure which is the best way.
-				DataObj = NewObject<UNiagaraCurveDataObject>(PinnedProps);
-				ScriptProps->ExternalConstants.SetOrAdd(ConstantName, DataObj);
-			}
-
-			UNiagaraCurveDataObject *CurveData = Cast<UNiagaraCurveDataObject>(DataObj);
-			CurveData->SetCurveObject(CurCurve);
-		}
 	}
 
 	UObject * GetCurve() const
 	{
-		UNiagaraDataObject *DataObj = ScriptProps->ExternalConstants.FindDataObj(ConstantName);
-		UNiagaraCurveDataObject *CurveData = Cast<UNiagaraCurveDataObject>(DataObj);
-		if (CurveData)
-		{
-			return CurveData->GetCurveObject();
-		}
 		return nullptr;
 	}
 
diff --git a/Engine/Source/Editor/NiagaraEditor/Public/EdGraphSchema_Niagara.h b/Engine/Source/Editor/NiagaraEditor/Public/EdGraphSchema_Niagara.h
new file mode 100644
--- /dev/null
+++ b/Engine/Source/Editor/NiagaraEditor/Public/EdGraphSchema_Niagara.h
@@ -0,0 +1,89 @@
+// Copyright 1998-2016 Epic Games, Inc. All Rights Reserved.
+
+#pragma once
+#include "EdGraph/EdGraphSchema.h"
+#include "INiagaraCompiler.h"
+#include "EdGraphSchema_Niagara.generated.h"
+
+/** Action to add a node to the graph */
+USTRUCT()
+struct NIAGARAEDITOR_API FNiagaraSchemaAction_NewNode : public FEdGraphSchemaAction
+{
+	GENERATED_USTRUCT_BODY();
+
+	/** Template of node we want to create */
+	UPROPERTY()
+	class UNiagaraNode* NodeTemplate;
+
+
+	FNiagaraSchemaAction_NewNode() 
+		: FEdGraphSchemaAction()
+		, NodeTemplate(NULL)
+	{}
+
+	FNiagaraSchemaAction_NewNode(const FText& InNodeCategory, const FText& InMenuDesc, const FString& InToolTip, const int32 InGrouping)
+		: FEdGraphSchemaAction(InNodeCategory, InMenuDesc, InToolTip, InGrouping) 
+		, NodeTemplate(NULL)
+	{}
+
+	//~ Begin FEdGraphSchemaAction Interface
+	virtual UEdGraphNode* PerformAction(class UEdGraph* ParentGraph, UEdGraphPin* FromPin, const FVector2D Location, bool bSelectNewNode = true) override;
+	virtual UEdGraphNode* PerformAction(class UEdGraph* ParentGraph, TArray<UEdGraphPin*>& FromPins, const FVector2D Location, bool bSelectNewNode = true) override;
+	virtual void AddReferencedObjects( FReferenceCollector& Collector ) override;
+	//~ End FEdGraphSchemaAction Interface
+
+	template <typename NodeType>
+	static NodeType* SpawnNodeFromTemplate(class UEdGraph* ParentGraph, NodeType* InTemplateNode, const FVector2D Location, bool bSelectNewNode = true)
+	{
+		FNiagaraSchemaAction_NewNode Action;
+		Action.NodeTemplate = InTemplateNode;
+
+		return Cast<NodeType>(Action.PerformAction(ParentGraph, NULL, Location, bSelectNewNode));
+	}
+};
+
+UCLASS()
+class NIAGARAEDITOR_API UEdGraphSchema_Niagara : public UEdGraphSchema
+{
+	GENERATED_UCLASS_BODY()
+
+	// Allowable PinType.PinCategory values
+	UPROPERTY()
+	FString PC_Float;
+	UPROPERTY()
+	FString PC_Vector;
+	UPROPERTY()
+	FString PC_Matrix;
+
+	//~ Begin EdGraphSchema Interface
+	virtual void GetGraphContextActions(FGraphContextMenuBuilder& ContextMenuBuilder) const override;
+	virtual void GetContextMenuActions(const UEdGraph* CurrentGraph, const UEdGraphNode* InGraphNode, const UEdGraphPin* InGraphPin, class FMenuBuilder* MenuBuilder, bool bIsDebugging) const override;
+	virtual const FPinConnectionResponse CanCreateConnection(const UEdGraphPin* A, const UEdGraphPin* B) const override;
+	virtual FLinearColor GetPinTypeColor(const FEdGraphPinType& PinType) const override;
+	virtual bool ShouldHidePinDefaultValue(UEdGraphPin* Pin) const override;
+	virtual bool TryCreateConnection(UEdGraphPin* A, UEdGraphPin* B) const override;
+
+	virtual void BreakSinglePinLink(UEdGraphPin* SourcePin, UEdGraphPin* TargetPin) override;
+	//~ End EdGraphSchema Interface
+
+	ENiagaraDataType GetPinDataType(UEdGraphPin* Pin)const;
+	void GetPinDefaultValue(UEdGraphPin* Pin, float& OutDefault)const;
+	void GetPinDefaultValue(UEdGraphPin* Pin, FVector4& OutDefault)const;
+	void GetPinDefaultValue(UEdGraphPin* Pin, FMatrix& OutDefault)const;
+		
+	bool IsSystemConstant(const FNiagaraVariableInfo& Variable)const;
+
+	virtual FString GetFloatName()const { return PC_Float; }
+	virtual FString GetVectorName()const { return PC_Vector; }
+	virtual FString GetMatrixName()const { return PC_Matrix; }
+
+	static const FLinearColor NodeTitleColor_Attribute;
+	static const FLinearColor NodeTitleColor_Constant;
+	static const FLinearColor NodeTitleColor_SystemConstant;
+	static const FLinearColor NodeTitleColor_FunctionCall;
+	static const FLinearColor NodeTitleColor_Event;
+
+private:
+	void GetBreakLinkToSubMenuActions(class FMenuBuilder& MenuBuilder, UEdGraphPin* InGraphPin);
+};
+
diff --git a/Engine/Source/Editor/NiagaraEditor/Public/INiagaraCompiler.h b/Engine/Source/Editor/NiagaraEditor/Public/INiagaraCompiler.h
new file mode 100644
--- /dev/null
+++ b/Engine/Source/Editor/NiagaraEditor/Public/INiagaraCompiler.h
@@ -0,0 +1,73 @@
+// Copyright 1998-2016 Epic Games, Inc. All Rights Reserved.
+
+#pragma once
+
+#include "NiagaraEditorCommon.h"
+
+class UNiagaraNode;
+
+//Interface for Niagara compilers.
+class INiagaraCompiler
+{
+public:
+
+	/** Compiles a script. */
+	virtual bool CompileScript(class UNiagaraScript* InScript) = 0;
+
+	/** Traverses the pins input, recursively compiling them and then adds any expressions needed for the owner node.*/
+	virtual TNiagaraExprPtr CompilePin(UEdGraphPin* Pin) = 0;
+	
+	/** Creates an expression which fetches the named attribute. */
+	virtual TNiagaraExprPtr GetAttribute(const FNiagaraVariableInfo& Attribute) = 0;
+	/** Creates an expression which fetches the named constant. */
+	virtual TNiagaraExprPtr GetExternalConstant(const FNiagaraVariableInfo& Constant, float Default) = 0;
+	/** Creates an expression which fetches the named constant. */
+	virtual TNiagaraExprPtr GetExternalConstant(const FNiagaraVariableInfo& Constant, const FVector4& Default) = 0;
+	/** Creates an expression which fetches the named constant. */
+	virtual TNiagaraExprPtr GetExternalConstant(const FNiagaraVariableInfo& Constant, const FMatrix& Default) = 0;
+
+	/** Creates an expression which fetches the named constant. */
+	virtual TNiagaraExprPtr GetInternalConstant(const FNiagaraVariableInfo& Constant, float Default) = 0;
+	/** Creates an expression which fetches the named constant. */
+	virtual TNiagaraExprPtr GetInternalConstant(const FNiagaraVariableInfo& Constant, const FVector4& Default) = 0;
+	/** Creates an expression which fetches the named constant. */
+	virtual TNiagaraExprPtr GetInternalConstant(const FNiagaraVariableInfo& Constant, const FMatrix& Default) = 0;
+
+	/** Creates an expression which writes the result of SourceExpression to the named attribute. */
+	virtual TNiagaraExprPtr Output(const FNiagaraVariableInfo& Attr, TNiagaraExprPtr& SourceExpression) = 0;
+
+	/** Creates an expression which reads data from a shared data set. */
+	virtual TNiagaraExprPtr SharedDataRead(const FNiagaraDataSetID& DataSet, const FNiagaraVariableInfo& Variable, TNiagaraExprPtr& IndexExpr) = 0;
+	/** Creates an expression which writes data to an shared data set. */
+	virtual TNiagaraExprPtr SharedDataWrite(const FNiagaraDataSetID& DataSet, const FNiagaraVariableInfo& Variable, TNiagaraExprPtr& IndexExpr, TNiagaraExprPtr& DataExpr) = 0;
+
+	/** Creates an expression which gets the current counter for the shared data set and increments the counter. */
+	virtual TNiagaraExprPtr AcquireSharedDataIndex(const FNiagaraDataSetID& DataSet, bool bWrap, TNiagaraExprPtr& ValidExpr) = 0;
+	/** Creates an expression which gets the current counter for the shared data set and decrements the counter. */
+	virtual TNiagaraExprPtr ConsumeSharedDataIndex(const FNiagaraDataSetID& DataSet, bool bWrap) = 0;
+	/** Creates an expression with a mask of 0 or 0xFFFFFFFF marking the passed index as valid or not.*/
+	virtual TNiagaraExprPtr SharedDataIndexIsValid(const FNiagaraDataSetID& DataSet, TNiagaraExprPtr& IndexExpr) = 0;
+
+	/** Checks the validity of the input expressions against the informations about this operation in FNiagaraOpInfo. */
+	virtual bool CheckInputs(FName OpName, TArray<TNiagaraExprPtr>& Inputs) = 0;
+	/** Checks the validity of the output expressions against the informations about this operation in FNiagaraOpInfo. */
+	virtual bool CheckOutputs(FName OpName, TArray<TNiagaraExprPtr>& Outputs) = 0;
+
+	/** Adds an error to be reported to the user. Any error will lead to compilation failure. */
+	virtual void Error(FText ErrorText, UNiagaraNode* Node, UEdGraphPin* Pin) = 0 ;
+
+	/** Adds a warning to be reported to the user. Warnings will not cause a compilation failure. */
+	virtual void Warning(FText WarningText, UNiagaraNode* Node, UEdGraphPin* Pin) = 0;
+
+	//Static function named after each op that allows access to internals via FNiagaraOpInfo.Func.
+#define NiagaraOp(OpName) static bool OpName(INiagaraCompiler* Compiler, TArray<TNiagaraExprPtr>& InputExpressions, TArray<TNiagaraExprPtr>& OutputExpressions);
+	NiagaraOpList;
+#undef NiagaraOp
+
+protected:
+	//Internal virtual function for each op. 
+#define NiagaraOp(OpName) virtual bool OpName##_Internal(TArray<TNiagaraExprPtr>& InputExpressions, TArray<TNiagaraExprPtr>& OutputExpressions) = 0;
+	NiagaraOpList;
+#undef NiagaraOp
+
+};
\ No newline at end of file
diff --git a/Engine/Source/Editor/NiagaraEditor/Public/NiagaraEffectFactoryNew.h b/Engine/Source/Editor/NiagaraEditor/Public/NiagaraEffectFactoryNew.h
new file mode 100644
--- /dev/null
+++ b/Engine/Source/Editor/NiagaraEditor/Public/NiagaraEffectFactoryNew.h
@@ -0,0 +1,17 @@
+// Copyright 1998-2016 Epic Games, Inc. All Rights Reserved.
+
+#pragma once
+#include "NiagaraEffectFactoryNew.generated.h"
+
+UCLASS(hidecategories = Object)
+class UNiagaraEffectFactoryNew : public UFactory
+{
+	GENERATED_UCLASS_BODY()
+
+	//~ Begin UFactory Interface
+	virtual UObject* FactoryCreateNew(UClass* Class, UObject* InParent, FName Name, EObjectFlags Flags, UObject* Context, FFeedbackContext* Warn) override;
+	//~ Begin UFactory Interface	
+};
+
+
+
diff --git a/Engine/Source/Editor/NiagaraEditor/Public/NiagaraGraph.h b/Engine/Source/Editor/NiagaraEditor/Public/NiagaraGraph.h
new file mode 100644
--- /dev/null
+++ b/Engine/Source/Editor/NiagaraEditor/Public/NiagaraGraph.h
@@ -0,0 +1,29 @@
+// Copyright 1998-2016 Epic Games, Inc. All Rights Reserved.
+
+#pragma once
+#include "EdGraph/EdGraph.h"
+#include "NiagaraGraph.generated.h"
+
+UCLASS(MinimalAPI)
+class UNiagaraGraph : public UEdGraph
+{
+	GENERATED_UCLASS_BODY()
+
+	//~ Begin UObject Interface
+	virtual void PostEditChangeProperty(struct FPropertyChangedEvent& PropertyChangedEvent)override;
+	//~ End UObjet Interface
+	
+	/** Get the source that owns this graph */
+	class UNiagaraScriptSource* GetSource() const;
+		
+	class UNiagaraNodeOutput* FindOutputNode() const;
+	void FindInputNodes(TArray<class UNiagaraNodeInput*>& OutInputNodes) const;
+
+	void FindReadDataSetNodes(TArray<class UNiagaraNodeReadDataSet*>& OutReadNodes) const;
+	void FindWriteDataSetNodes(TArray<class UNiagaraNodeWriteDataSet*>& OutWriteNodes) const;
+
+	/** Returns the index of this attribute in the output node of the graph. INDEX_NONE if this is not a valid attribute. */
+	int32 GetAttributeIndex(const FNiagaraVariableInfo& Attr)const;
+	void GetAttributes(TArray< FNiagaraVariableInfo >& OutAttributes)const;
+};
+
diff --git a/Engine/Source/Editor/NiagaraEditor/Public/NiagaraNode.h b/Engine/Source/Editor/NiagaraEditor/Public/NiagaraNode.h
new file mode 100644
--- /dev/null
+++ b/Engine/Source/Editor/NiagaraEditor/Public/NiagaraNode.h
@@ -0,0 +1,34 @@
+// Copyright 1998-2016 Epic Games, Inc. All Rights Reserved.
+
+#pragma once
+#include "NiagaraNode.generated.h"
+
+UCLASS()
+class NIAGARAEDITOR_API UNiagaraNode : public UEdGraphNode
+{
+	GENERATED_UCLASS_BODY()
+protected:
+
+	void ReallocatePins();
+
+public:
+
+	//~ Begin EdGraphNode Interface
+	virtual void AutowireNewNode(UEdGraphPin* FromPin)override;
+	//~ End EdGraphNode Interface
+
+	/** Get the Niagara graph that owns this node */
+	const class UNiagaraGraph* GetNiagaraGraph()const;
+	class UNiagaraGraph* GetNiagaraGraph();
+
+	/** Get the source object */
+	class UNiagaraScriptSource* GetSource()const;
+
+	virtual void Compile(class INiagaraCompiler* Compiler, TArray<FNiagaraNodeResult>& Outputs){ check(0); }
+
+	UEdGraphPin* GetInputPin(int32 InputIndex) const;
+	void GetInputPins(TArray<class UEdGraphPin*>& OutInputPins) const;
+	UEdGraphPin* GetOutputPin(int32 OutputIndex) const;
+	void GetOutputPins(TArray<class UEdGraphPin*>& OutOutputPins) const;
+
+};
diff --git a/Engine/Source/Editor/NiagaraEditor/Public/NiagaraNodeConstant.h b/Engine/Source/Editor/NiagaraEditor/Public/NiagaraNodeConstant.h
new file mode 100644
--- /dev/null
+++ b/Engine/Source/Editor/NiagaraEditor/Public/NiagaraNodeConstant.h
@@ -0,0 +1,6 @@
+// Copyright 1998-2016 Epic Games, Inc. All Rights Reserved.
+
+#pragma once
+
+#include "NiagaraCommon.h"
+#include "NiagaraNodeConstant.generated.h"
diff --git a/Engine/Source/Editor/NiagaraEditor/Public/NiagaraNodeGetAttr.h b/Engine/Source/Editor/NiagaraEditor/Public/NiagaraNodeGetAttr.h
new file mode 100644
--- /dev/null
+++ b/Engine/Source/Editor/NiagaraEditor/Public/NiagaraNodeGetAttr.h
@@ -0,0 +1,4 @@
+// Copyright 1998-2016 Epic Games, Inc. All Rights Reserved.
+
+#pragma once
+#include "NiagaraNodeGetAttr.generated.h"
diff --git a/Engine/Source/Editor/NiagaraEditor/Public/NiagaraNodeInput.h b/Engine/Source/Editor/NiagaraEditor/Public/NiagaraNodeInput.h
new file mode 100644
--- /dev/null
+++ b/Engine/Source/Editor/NiagaraEditor/Public/NiagaraNodeInput.h
@@ -0,0 +1,56 @@
+// Copyright 1998-2016 Epic Games, Inc. All Rights Reserved.
+
+#pragma once
+
+#include "NiagaraNode.h"
+#include "NiagaraNodeInput.generated.h"
+
+UCLASS(MinimalAPI)
+class UNiagaraNodeInput : public UNiagaraNode
+{
+	GENERATED_UCLASS_BODY()
+
+public:
+
+	UPROPERTY(EditAnywhere, Category = Input)
+	FNiagaraVariableInfo Input;
+
+	//TODO: Customize the details for this and hide these when they're not relevant.
+
+	UPROPERTY(EditAnywhere, Category = Constant)
+	float FloatDefault;
+
+	UPROPERTY(EditAnywhere, Category = Constant)
+	FVector4 VectorDefault;
+
+	UPROPERTY(EditAnywhere, Category = Constant)
+	FMatrix MatrixDefault;
+
+	/** Allows code to explicitly disable exposing of certain inputs e.g. system constants such as Delta Time. */
+	UPROPERTY(EditAnywhere, Category=Constant, meta=(InlineEditConditionToggle))
+	uint32 bCanBeExposed:1;
+
+	/** When true, and this input is a constant, the input is exposed to the effect editor. */
+	UPROPERTY(EditAnywhere, Category = Constant, meta = (editcondition = "bCanBeExposed"))
+	uint32 bExposeWhenConstant:1;
+
+	//~ Begin UObject Interface
+	virtual void PostEditChangeProperty(struct FPropertyChangedEvent& PropertyChangedEvent) override;
+	//~ End UObject Interface
+
+	//~ Begin EdGraphNode Interface
+	virtual void AllocateDefaultPins() override;
+	virtual FText GetNodeTitle(ENodeTitleType::Type TitleType) const override;
+	virtual FLinearColor GetNodeTitleColor() const override;
+	virtual void AutowireNewNode(UEdGraphPin* FromPin)override;
+	//~ End EdGraphNode Interface
+
+	virtual void Compile(class INiagaraCompiler* Compiler, TArray<FNiagaraNodeResult>& Outputs)override;
+	
+	bool IsConstant()const;
+	bool IsExposedConstant()const;
+
+	static const FLinearColor TitleColor_Attribute;
+	static const FLinearColor TitleColor_Constant;
+};
+
diff --git a/Engine/Source/Editor/NiagaraEditor/Public/NiagaraNodeOp.h b/Engine/Source/Editor/NiagaraEditor/Public/NiagaraNodeOp.h
new file mode 100644
--- /dev/null
+++ b/Engine/Source/Editor/NiagaraEditor/Public/NiagaraNodeOp.h
@@ -0,0 +1,83 @@
+// Copyright 1998-2016 Epic Games, Inc. All Rights Reserved.
+
+#pragma once
+#include "NiagaraNode.h"
+#include "NiagaraNodeOp.generated.h"
+
+UCLASS(MinimalAPI)
+class UNiagaraNodeOp : public UNiagaraNode
+{
+	GENERATED_UCLASS_BODY()
+
+public:
+
+	/** Name of operation */
+	UPROPERTY()
+	FName OpName;
+
+	//~ Begin EdGraphNode Interface
+	virtual void AllocateDefaultPins() override;
+	virtual FText GetNodeTitle(ENodeTitleType::Type TitleType) const override;
+	virtual FText GetTooltipText() const override;
+	virtual FLinearColor GetNodeTitleColor() const override;
+	//~ End EdGraphNode Interface
+
+	//~ Begin UNiagaraNode Interface
+	virtual void Compile(class INiagaraCompiler* Compiler, TArray<FNiagaraNodeResult>& Outputs) override;
+	//~ End UNiagaraNode Interface
+};
+
+
+
+#if 0
+
+class SGraphNodeNiagara : public SGraphNode
+{
+public:
+	SLATE_BEGIN_ARGS(SGraphNodeNiagara){}
+	SLATE_END_ARGS()
+
+		void Construct(const FArguments& InArgs, class UNiagaraNode* InNode);
+
+	//~ Begin SGraphNode Interface
+	virtual void CreatePinWidgets() override;
+	//~ End SGraphNode Interface
+
+	//~ Begin SNodePanel::SNode Interface
+	virtual void MoveTo(const FVector2D& NewPosition, FNodeSet& NodeFilter) override;
+	//~ End SNodePanel::SNode Interface
+
+	UNiagaraNode* GetNiagaraGraphNode() const { return NiagaraNode; }
+
+protected:
+	//~ Begin SGraphNode Interface
+	virtual void AddPin(const TSharedRef<SGraphPin>& PinToAdd) override;
+	virtual void CreateBelowPinControls(TSharedPtr<SVerticalBox> MainBox) override;
+	virtual void SetDefaultTitleAreaWidget(TSharedRef<SOverlay> DefaultTitleAreaWidget) override;
+	virtual TSharedRef<SWidget> CreateNodeContentArea() override;
+	//~ End SGraphNode Interface
+
+	/** Creates a preview viewport if necessary */
+	//TSharedRef<SWidget> CreatePreviewWidget();
+
+	/** Returns visibility of Expression Preview viewport */
+	//EVisibility ExpressionPreviewVisibility() const;
+
+	/** Show/hide Expression Preview */
+	//void OnExpressionPreviewChanged(const ECheckBoxState NewCheckedState);
+
+	/** hidden == unchecked, shown == checked */
+	//ECheckBoxState IsExpressionPreviewChecked() const;
+
+	/** Up when shown, down when hidden */
+	//const FSlateBrush* GetExpressionPreviewArrow() const;
+
+private:
+	/** Slate viewport for rendering preview via custom slate element */
+	//TSharedPtr<FPreviewViewport> PreviewViewport;
+
+	/** Cached material graph node pointer to avoid casting */
+	UNiagaraNode* NiagaraNode;
+};
+
+#endif
\ No newline at end of file
diff --git a/Engine/Source/Editor/NiagaraEditor/Public/NiagaraNodeOutput.h b/Engine/Source/Editor/NiagaraEditor/Public/NiagaraNodeOutput.h
new file mode 100644
--- /dev/null
+++ b/Engine/Source/Editor/NiagaraEditor/Public/NiagaraNodeOutput.h
@@ -0,0 +1,31 @@
+// Copyright 1998-2016 Epic Games, Inc. All Rights Reserved.
+
+#pragma once
+
+#include "NiagaraNode.h"
+#include "NiagaraNodeOutput.generated.h"
+
+UCLASS(MinimalAPI)
+class UNiagaraNodeOutput : public UNiagaraNode
+{
+	GENERATED_UCLASS_BODY()
+
+	UPROPERTY(EditAnywhere, Category = Output)
+	TArray<FNiagaraVariableInfo> Outputs;
+
+public:
+
+	//~ Begin UObject Interface
+	virtual void PostEditChangeProperty(struct FPropertyChangedEvent& PropertyChangedEvent) override;
+	//~ End UObject Interface
+
+	//~ Begin EdGraphNode Interface
+	virtual void AllocateDefaultPins() override;
+	virtual bool CanUserDeleteNode() const override { return false; }
+	virtual FText GetNodeTitle(ENodeTitleType::Type TitleType) const override;
+	virtual FLinearColor GetNodeTitleColor() const override;
+	//~ End EdGraphNode Interface
+
+	virtual void Compile(class INiagaraCompiler* Compiler, TArray<FNiagaraNodeResult>& OutputExpressions)override;
+};
+
diff --git a/Engine/Source/Editor/NiagaraEditor/Public/NiagaraNodeOutputUpdate.h b/Engine/Source/Editor/NiagaraEditor/Public/NiagaraNodeOutputUpdate.h
new file mode 100644
--- /dev/null
+++ b/Engine/Source/Editor/NiagaraEditor/Public/NiagaraNodeOutputUpdate.h
@@ -0,0 +1,6 @@
+// Copyright 1998-2016 Epic Games, Inc. All Rights Reserved.
+
+#pragma once
+#include "NiagaraNode.h"
+#include "NiagaraNodeOutputUpdate.generated.h"
+
diff --git a/Engine/Source/Editor/NiagaraEditor/Public/NiagaraNodeReadDataSet.h b/Engine/Source/Editor/NiagaraEditor/Public/NiagaraNodeReadDataSet.h
new file mode 100644
--- /dev/null
+++ b/Engine/Source/Editor/NiagaraEditor/Public/NiagaraNodeReadDataSet.h
@@ -0,0 +1,32 @@
+// Copyright 1998-2016 Epic Games, Inc. All Rights Reserved.
+
+#pragma once
+#include "NiagaraNode.h"
+#include "NiagaraNodeReadDataSet.generated.h"
+
+UCLASS(MinimalAPI)
+class UNiagaraNodeReadDataSet : public UNiagaraNode
+{
+	GENERATED_UCLASS_BODY()
+
+public:
+
+	UPROPERTY(EditAnywhere, Category = DataSet)
+	FNiagaraDataSetID DataSet;
+
+	UPROPERTY(EditAnywhere, Category = Variables)
+	TArray<FNiagaraVariableInfo> Variables;
+
+	//~ Begin UObject Interface
+	virtual void PostEditChangeProperty(struct FPropertyChangedEvent& PropertyChangedEvent) override;
+	//~ End UObject Interface
+
+	//~ Begin EdGraphNode Interface
+	virtual void AllocateDefaultPins() override;
+	virtual FText GetNodeTitle(ENodeTitleType::Type TitleType) const override;
+	virtual FLinearColor GetNodeTitleColor() const override;
+	//~ End EdGraphNode Interface
+
+	virtual void Compile(class INiagaraCompiler* Compiler, TArray<FNiagaraNodeResult>& Outputs)override;
+};
+
diff --git a/Engine/Source/Editor/NiagaraEditor/Public/NiagaraNodeWriteDataSet.h b/Engine/Source/Editor/NiagaraEditor/Public/NiagaraNodeWriteDataSet.h
new file mode 100644
--- /dev/null
+++ b/Engine/Source/Editor/NiagaraEditor/Public/NiagaraNodeWriteDataSet.h
@@ -0,0 +1,32 @@
+// Copyright 1998-2016 Epic Games, Inc. All Rights Reserved.
+
+#pragma once
+#include "NiagaraNode.h"
+#include "NiagaraNodeWriteDataSet.generated.h"
+
+UCLASS(MinimalAPI)
+class UNiagaraNodeWriteDataSet : public UNiagaraNode
+{
+	GENERATED_UCLASS_BODY()
+
+public:
+
+	UPROPERTY(EditAnywhere, Category = DataSet)
+	FNiagaraDataSetID DataSet;
+
+	UPROPERTY(EditAnywhere, Category = Variables)
+	TArray<FNiagaraVariableInfo> Variables;
+
+	//~ Begin UObject Interface
+	virtual void PostEditChangeProperty(struct FPropertyChangedEvent& PropertyChangedEvent) override;
+	//~ End UObject Interface
+
+	//~ Begin EdGraphNode Interface
+	virtual void AllocateDefaultPins() override;
+	virtual FText GetNodeTitle(ENodeTitleType::Type TitleType) const override;
+	virtual FLinearColor GetNodeTitleColor() const override;
+	//~ End EdGraphNode Interface
+
+	virtual void Compile(class INiagaraCompiler* Compiler, TArray<FNiagaraNodeResult>& Outputs)override;
+};
+
diff --git a/Engine/Source/Editor/NiagaraEditor/Public/NiagaraScriptFactoryNew.h b/Engine/Source/Editor/NiagaraEditor/Public/NiagaraScriptFactoryNew.h
new file mode 100644
--- /dev/null
+++ b/Engine/Source/Editor/NiagaraEditor/Public/NiagaraScriptFactoryNew.h
@@ -0,0 +1,17 @@
+// Copyright 1998-2016 Epic Games, Inc. All Rights Reserved.
+
+#pragma once
+#include "NiagaraScriptFactoryNew.generated.h"
+
+UCLASS(hidecategories=Object)
+class UNiagaraScriptFactoryNew : public UFactory
+{
+	GENERATED_UCLASS_BODY()
+
+	//~ Begin UFactory Interface
+	virtual UObject* FactoryCreateNew(UClass* Class,UObject* InParent,FName Name,EObjectFlags Flags,UObject* Context,FFeedbackContext* Warn) override;
+	//~ Begin UFactory Interface	
+};
+
+
+
diff --git a/Engine/Source/Editor/NiagaraEditor/Public/NiagaraScriptSource.h b/Engine/Source/Editor/NiagaraEditor/Public/NiagaraScriptSource.h
new file mode 100644
--- /dev/null
+++ b/Engine/Source/Editor/NiagaraEditor/Public/NiagaraScriptSource.h
@@ -0,0 +1,27 @@
+// Copyright 1998-2016 Epic Games, Inc. All Rights Reserved.
+
+
+#pragma once
+#include "NiagaraScriptSourceBase.h"
+#include "NiagaraScriptSource.generated.h"
+
+
+UCLASS(MinimalAPI)
+class UNiagaraScriptSource : public UNiagaraScriptSourceBase
+{
+	GENERATED_UCLASS_BODY()
+
+	/** Graph for particle update expression */
+	UPROPERTY()
+	class UNiagaraGraph*	NodeGraph;
+
+	/** The same node graph from above but with all function calls merge into a single graph. */
+	UPROPERTY()
+	class UNiagaraGraph*	FlattenedNodeGraph;
+
+	// UObject interface.
+	virtual void PostLoad() override;
+	virtual void Compile() override;
+
+	void GetEmitterAttributes(TArray<FName>& VectorInputs, TArray<FName>& MatrixInputs);
+};
diff --git a/Engine/Source/Editor/SequencerWidgets/SequencerWidgets.Build.cs b/Engine/Source/Editor/SequencerWidgets/SequencerWidgets.Build.cs
--- a/Engine/Source/Editor/SequencerWidgets/SequencerWidgets.Build.cs
+++ b/Engine/Source/Editor/SequencerWidgets/SequencerWidgets.Build.cs
@@ -8,7 +8,19 @@ public class SequencerWidgets : ModuleRules
 	{
 		PrivateIncludePaths.Add("Editor/SequencerWidgets/Private");
 
-		PrivateDependencyModuleNames.AddRange(
+        PrivateIncludePathModuleNames.AddRange(
+            new string[] {
+                "MovieScene",
+            }
+           );
+
+        PublicDependencyModuleNames.AddRange(
+            new string[] {
+                "MovieScene",
+            }
+        );
+        
+        PrivateDependencyModuleNames.AddRange(
 			new string[] {
 				"Core",
 				"CoreUObject",
@@ -18,7 +30,8 @@ public class SequencerWidgets : ModuleRules
                 "MovieScene",
 				"Slate",
 				"SlateCore",
-				"UnrealEd"
+                "MovieScene",
+                "UnrealEd"
 			}
 		);
 	}
diff --git a/Engine/Source/Editor/UnrealEd/Classes/Factories/TextureFactory.h b/Engine/Source/Editor/UnrealEd/Classes/Factories/TextureFactory.h
--- a/Engine/Source/Editor/UnrealEd/Classes/Factories/TextureFactory.h
+++ b/Engine/Source/Editor/UnrealEd/Classes/Factories/TextureFactory.h
@@ -77,6 +77,10 @@ class UNREALED_API UTextureFactory : public UFactory
 	UPROPERTY(EditAnywhere, Category=DitherMipMaps, meta=(ToolTip="If enabled, mip-map alpha values will be dithered for smooth transitions"))
 	uint32 bDitherMipMapAlpha:1;
 
+	/** Channel values to compare to when preserving alpha coverage from a mask. */
+	UPROPERTY(EditAnywhere, Category=PreserveAlphaCoverage, meta=(ToolTip="Channel values to compare to when preserving alpha coverage from a mask for mips"))
+	FVector4 AlphaCoverageThresholds;
+
 	/** If enabled, preserve the value of border pixels when creating mip-maps */
 	UPROPERTY(EditAnywhere, Category=PreserveBorder, meta=(ToolTip="If enabled, preserve the value of border pixels when creating mip-maps"))
 	uint32 bPreserveBorder:1;
