Fix nonunity compile

diff --git a/Engine/Source/Editor/NiagaraEditor/Private/Widgets/SNiagaraScriptGraph.h b/Engine/Source/Editor/NiagaraEditor/Private/Widgets/SNiagaraScriptGraph.h
--- a/Engine/Source/Editor/NiagaraEditor/Private/Widgets/SNiagaraScriptGraph.h
+++ b/Engine/Source/Editor/NiagaraEditor/Private/Widgets/SNiagaraScriptGraph.h
@@ -2,6 +2,8 @@
 
 #pragma once
 
+#include "SCompoundWidget.h"
+
 class FNiagaraScriptGraphViewModel;
 
 /** A widget for editing a UNiagaraScript with a graph. */
