Fixed non-unity build

diff --git a/Engine/Source/Editor/NiagaraEditor/Private/NiagaraNodeOutput.cpp b/Engine/Source/Editor/NiagaraEditor/Private/NiagaraNodeOutput.cpp
--- a/Engine/Source/Editor/NiagaraEditor/Private/NiagaraNodeOutput.cpp
+++ b/Engine/Source/Editor/NiagaraEditor/Private/NiagaraNodeOutput.cpp
@@ -3,6 +3,7 @@
 #include "NiagaraNodeOutput.h"
 #include "UObject/UnrealType.h"
 #include "INiagaraCompiler.h"
+#include "NiagaraScript.h"
 
 #include "EdGraphSchema_Niagara.h"
 
