Fix another non-unity

diff --git a/Engine/Source/Editor/NiagaraEditor/Private/Widgets/SNiagaraSelectedObjectsDetails.h b/Engine/Source/Editor/NiagaraEditor/Private/Widgets/SNiagaraSelectedObjectsDetails.h
--- a/Engine/Source/Editor/NiagaraEditor/Private/Widgets/SNiagaraSelectedObjectsDetails.h
+++ b/Engine/Source/Editor/NiagaraEditor/Private/Widgets/SNiagaraSelectedObjectsDetails.h
@@ -1,6 +1,7 @@
 // Copyright 1998-2016 Epic Games, Inc. All Rights Reserved.
 
 #pragma once
+#include "SCompoundWidget.h"
 
 class FNiagaraObjectSelection;
 class IDetailsView;
