Comprehensive Editor/Monolithic non-unity IWYU fixes for Niagara

diff --git a/Engine/Source/Editor/NiagaraEditor/Private/Customizations/NiagaraComponentDetails.cpp b/Engine/Source/Editor/NiagaraEditor/Private/Customizations/NiagaraComponentDetails.cpp
--- a/Engine/Source/Editor/NiagaraEditor/Private/Customizations/NiagaraComponentDetails.cpp
+++ b/Engine/Source/Editor/NiagaraEditor/Private/Customizations/NiagaraComponentDetails.cpp
@@ -27,6 +27,7 @@
 #include "IDetailsView.h"
 #include "Widgets/Input/SButton.h"
 #include "Widgets/Input/SEditableText.h"
+#include "Editor.h"
 #define LOCTEXT_NAMESPACE "NiagaraComponentDetails"
 
 class FNiagaraParameterViewModelCustomDetails : public FNiagaraParameterViewModel
diff --git a/Engine/Source/Editor/NiagaraEditor/Private/Customizations/NiagaraComponentDetails.h b/Engine/Source/Editor/NiagaraEditor/Private/Customizations/NiagaraComponentDetails.h
--- a/Engine/Source/Editor/NiagaraEditor/Private/Customizations/NiagaraComponentDetails.h
+++ b/Engine/Source/Editor/NiagaraEditor/Private/Customizations/NiagaraComponentDetails.h
@@ -6,6 +6,8 @@
 #include "IDetailCustomization.h"
 #include "EditorUndoClient.h"
 #include "TickableEditorObject.h"
+#include "Reply.h"
+#include "Visibility.h"
 
 class IDetailLayoutBuilder;
 class IPropertyHandle;
diff --git a/Engine/Source/Editor/NiagaraEditor/Private/EdGraphSchema_Niagara.cpp b/Engine/Source/Editor/NiagaraEditor/Private/EdGraphSchema_Niagara.cpp
--- a/Engine/Source/Editor/NiagaraEditor/Private/EdGraphSchema_Niagara.cpp
+++ b/Engine/Source/Editor/NiagaraEditor/Private/EdGraphSchema_Niagara.cpp
@@ -25,6 +25,7 @@
 #include "NiagaraEditorUtilities.h"
 #include "NiagaraDataInterface.h"
 #include "NiagaraNodeIf.h"
+#include "MessageDialog.h"
 
 #include "ModuleManager.h"
 #include "AssetRegistryModule.h"
diff --git a/Engine/Source/Editor/NiagaraEditor/Private/NiagaraCompiler.cpp b/Engine/Source/Editor/NiagaraEditor/Private/NiagaraCompiler.cpp
--- a/Engine/Source/Editor/NiagaraEditor/Private/NiagaraCompiler.cpp
+++ b/Engine/Source/Editor/NiagaraEditor/Private/NiagaraCompiler.cpp
@@ -21,6 +21,8 @@
 
 #include "ShaderFormatVectorVM.h"
 
+#include "ShaderCore.h"
+
 #define LOCTEXT_NAMESPACE "NiagaraCompiler"
 
 DEFINE_LOG_CATEGORY_STATIC(LogNiagaraCompiler, All, All);
diff --git a/Engine/Source/Editor/NiagaraEditor/Private/NiagaraEditorCommands.h b/Engine/Source/Editor/NiagaraEditor/Private/NiagaraEditorCommands.h
--- a/Engine/Source/Editor/NiagaraEditor/Private/NiagaraEditorCommands.h
+++ b/Engine/Source/Editor/NiagaraEditor/Private/NiagaraEditorCommands.h
@@ -1,7 +1,8 @@
 // Copyright 1998-2016 Epic Games, Inc. All Rights Reserved.
 
 #pragma once
-
+#include "Commands.h"
+#include "EditorStyleSet.h"
 /**
 * Defines commands for the niagara editor.
 */
diff --git a/Engine/Source/Editor/NiagaraEditor/Private/NiagaraEditorStyle.cpp b/Engine/Source/Editor/NiagaraEditor/Private/NiagaraEditorStyle.cpp
--- a/Engine/Source/Editor/NiagaraEditor/Private/NiagaraEditorStyle.cpp
+++ b/Engine/Source/Editor/NiagaraEditor/Private/NiagaraEditorStyle.cpp
@@ -2,8 +2,11 @@
 
 #include "NiagaraEditorStyle.h"
 
+#include "SlateApplication.h"
+#include "EditorStyleSet.h"
 #include "Slate/SlateGameResources.h"
 #include "Styling/SlateStyleRegistry.h"
+#include "Styling/SlateTypes.h"
 
 TSharedPtr< FSlateStyleSet > FNiagaraEditorStyle::NiagaraEditorStyleInstance = NULL;
 
diff --git a/Engine/Source/Editor/NiagaraEditor/Private/NiagaraEditorUtilities.cpp b/Engine/Source/Editor/NiagaraEditor/Private/NiagaraEditorUtilities.cpp
--- a/Engine/Source/Editor/NiagaraEditor/Private/NiagaraEditorUtilities.cpp
+++ b/Engine/Source/Editor/NiagaraEditor/Private/NiagaraEditorUtilities.cpp
@@ -6,6 +6,8 @@
 #include "NiagaraNodeInput.h"
 #include "NiagaraDataInterface.h"
 #include "NiagaraComponent.h"
+#include "ModuleManager.h"
+#include "StructOnScope.h"
 
 FName FNiagaraEditorUtilities::GetUniqueName(FName CandidateName, const TSet<FName>& ExistingNames)
 {
diff --git a/Engine/Source/Editor/NiagaraEditor/Private/NiagaraEmitterFactoryNew.cpp b/Engine/Source/Editor/NiagaraEditor/Private/NiagaraEmitterFactoryNew.cpp
--- a/Engine/Source/Editor/NiagaraEditor/Private/NiagaraEmitterFactoryNew.cpp
+++ b/Engine/Source/Editor/NiagaraEditor/Private/NiagaraEmitterFactoryNew.cpp
@@ -5,6 +5,8 @@
 #include "NiagaraEditorModule.h"
 #include "NiagaraScriptFactoryNew.h"
 
+#include "ConfigCacheIni.h"
+
 #define LOCTEXT_NAMESPACE "NiagaraEmitterFactory"
 
 UNiagaraEmitterFactoryNew::UNiagaraEmitterFactoryNew(const FObjectInitializer& ObjectInitializer)
diff --git a/Engine/Source/Editor/NiagaraEditor/Private/NiagaraEmitterPropertiesDetailsCustomization.cpp b/Engine/Source/Editor/NiagaraEditor/Private/NiagaraEmitterPropertiesDetailsCustomization.cpp
--- a/Engine/Source/Editor/NiagaraEditor/Private/NiagaraEmitterPropertiesDetailsCustomization.cpp
+++ b/Engine/Source/Editor/NiagaraEditor/Private/NiagaraEmitterPropertiesDetailsCustomization.cpp
@@ -15,6 +15,12 @@
 #include "IDetailPropertyRow.h"
 #include "DetailCategoryBuilder.h"
 #include "PropertyCustomizationHelpers.h"
+#include "ScopedTransaction.h"
+
+#include "SBorder.h"
+#include "STextBlock.h"
+#include "SButton.h"
+
 
 #define LOCTEXT_NAMESPACE "NiagaraEffectEditor"
 
diff --git a/Engine/Source/Editor/NiagaraEditor/Private/NiagaraNodeEmitter.cpp b/Engine/Source/Editor/NiagaraEditor/Private/NiagaraNodeEmitter.cpp
--- a/Engine/Source/Editor/NiagaraEditor/Private/NiagaraNodeEmitter.cpp
+++ b/Engine/Source/Editor/NiagaraEditor/Private/NiagaraNodeEmitter.cpp
@@ -5,6 +5,7 @@
 #include "NiagaraEmitterProperties.h"
 #include "NiagaraEditorUtilities.h"
 #include "EdGraphSchema_Niagara.h"
+#include "EdGraph/EdGraph.h"
 
 #define LOCTEXT_NAMESPACE "NiagaraNodeEmitter"
 
diff --git a/Engine/Source/Editor/NiagaraEditor/Private/NiagaraNodeWithDynamicPins.cpp b/Engine/Source/Editor/NiagaraEditor/Private/NiagaraNodeWithDynamicPins.cpp
--- a/Engine/Source/Editor/NiagaraEditor/Private/NiagaraNodeWithDynamicPins.cpp
+++ b/Engine/Source/Editor/NiagaraEditor/Private/NiagaraNodeWithDynamicPins.cpp
@@ -4,7 +4,11 @@
 #include "NiagaraNodeWithDynamicPins.h"
 #include "EdGraphSchema_Niagara.h"
 
+#include "UIAction.h"
 #include "ScopedTransaction.h"
+#include "MultiBoxBuilder.h"
+#include "SEditableTextBox.h"
+#include "SBox.h"
 
 #define LOCTEXT_NAMESPACE "NiagaraNodeWithDynamicPins"
 
diff --git a/Engine/Source/Editor/NiagaraEditor/Private/NiagaraNodeWithDynamicPins.h b/Engine/Source/Editor/NiagaraEditor/Private/NiagaraNodeWithDynamicPins.h
--- a/Engine/Source/Editor/NiagaraEditor/Private/NiagaraNodeWithDynamicPins.h
+++ b/Engine/Source/Editor/NiagaraEditor/Private/NiagaraNodeWithDynamicPins.h
@@ -3,6 +3,8 @@
 #pragma once
 
 #include "NiagaraNode.h"
+#include "SlateTypes.h"
+#include "EdGraph/EdGraph.h"
 
 #include "NiagaraNodeWithDynamicPins.generated.h"
 
diff --git a/Engine/Source/Editor/NiagaraEditor/Private/Toolkits/NiagaraEffectToolkit.cpp b/Engine/Source/Editor/NiagaraEditor/Private/Toolkits/NiagaraEffectToolkit.cpp
--- a/Engine/Source/Editor/NiagaraEditor/Private/Toolkits/NiagaraEffectToolkit.cpp
+++ b/Engine/Source/Editor/NiagaraEditor/Private/Toolkits/NiagaraEffectToolkit.cpp
@@ -15,6 +15,15 @@
 #include "IContentBrowserSingleton.h"
 #include "ContentBrowserModule.h"
 
+#include "EditorStyleSet.h"
+#include "AssetEditorToolkit.h"
+#include "WorkspaceItem.h"
+
+#include "SlateApplication.h"
+#include "SBoxPanel.h"
+#include "SBox.h"
+#include "SDockTab.h"
+
 #define LOCTEXT_NAMESPACE "NiagaraEffectEditor"
 
 const FName FNiagaraEffectToolkit::ViewportTabID(TEXT("NiagaraEffectEditor_Viewport"));
diff --git a/Engine/Source/Editor/NiagaraEditor/Private/Toolkits/NiagaraEffectToolkit.h b/Engine/Source/Editor/NiagaraEditor/Private/Toolkits/NiagaraEffectToolkit.h
--- a/Engine/Source/Editor/NiagaraEditor/Private/Toolkits/NiagaraEffectToolkit.h
+++ b/Engine/Source/Editor/NiagaraEditor/Private/Toolkits/NiagaraEffectToolkit.h
@@ -8,6 +8,8 @@
 #include "UObject/GCObject.h"
 #include "Toolkits/IToolkitHost.h"
 
+#include "AssetEditorToolkit.h"
+
 #include "ISequencer.h"
 #include "ISequencerTrackEditor.h"
 
@@ -20,6 +22,7 @@ class SNiagaraEffectViewport;
 class SNiagaraEffectEditor;
 class UNiagaraEffect;
 class UNiagaraSequence;
+class FAssetData;
 
 /** Viewer/editor for a NiagaraEffect
 */
diff --git a/Engine/Source/Editor/NiagaraEditor/Private/Toolkits/NiagaraScriptToolkit.cpp b/Engine/Source/Editor/NiagaraEditor/Private/Toolkits/NiagaraScriptToolkit.cpp
--- a/Engine/Source/Editor/NiagaraEditor/Private/Toolkits/NiagaraScriptToolkit.cpp
+++ b/Engine/Source/Editor/NiagaraEditor/Private/Toolkits/NiagaraScriptToolkit.cpp
@@ -15,7 +15,11 @@
 #include "Package.h"
 #include "NiagaraEditorStyle.h"
 
+#include "ModuleManager.h"
+
+#include "MultiBoxBuilder.h"
 #include "SScrollBox.h"
+#include "SDockTab.h"
 
 #define LOCTEXT_NAMESPACE "NiagaraScriptToolkit"
 
diff --git a/Engine/Source/Editor/NiagaraEditor/Private/Toolkits/NiagaraScriptToolkit.h b/Engine/Source/Editor/NiagaraEditor/Private/Toolkits/NiagaraScriptToolkit.h
--- a/Engine/Source/Editor/NiagaraEditor/Private/Toolkits/NiagaraScriptToolkit.h
+++ b/Engine/Source/Editor/NiagaraEditor/Private/Toolkits/NiagaraScriptToolkit.h
@@ -7,6 +7,7 @@
 #include "Toolkits/IToolkitHost.h"
 #include "Misc/NotifyHook.h"
 #include "EditorUndoClient.h"
+#include "AssetEditorToolkit.h"
 
 class IDetailsView;
 class SGraphEditor;
diff --git a/Engine/Source/Editor/NiagaraEditor/Private/TypeEditorUtilities/NiagaraBoolTypeEditorUtilities.cpp b/Engine/Source/Editor/NiagaraEditor/Private/TypeEditorUtilities/NiagaraBoolTypeEditorUtilities.cpp
--- a/Engine/Source/Editor/NiagaraEditor/Private/TypeEditorUtilities/NiagaraBoolTypeEditorUtilities.cpp
+++ b/Engine/Source/Editor/NiagaraEditor/Private/TypeEditorUtilities/NiagaraBoolTypeEditorUtilities.cpp
@@ -3,6 +3,9 @@
 #include "NiagaraBoolTypeEditorUtilities.h"
 #include "NiagaraTypes.h"
 #include "SNiagaraParameterEditor.h"
+#include "DeclarativeSyntaxSupport.h"
+#include "SBoxPanel.h"
+#include "SCheckBox.h"
 
 class SNiagaraBoolParameterEditor : public SNiagaraParameterEditor
 {
diff --git a/Engine/Source/Editor/NiagaraEditor/Private/TypeEditorUtilities/NiagaraColorTypeEditorUtilities.cpp b/Engine/Source/Editor/NiagaraEditor/Private/TypeEditorUtilities/NiagaraColorTypeEditorUtilities.cpp
--- a/Engine/Source/Editor/NiagaraEditor/Private/TypeEditorUtilities/NiagaraColorTypeEditorUtilities.cpp
+++ b/Engine/Source/Editor/NiagaraEditor/Private/TypeEditorUtilities/NiagaraColorTypeEditorUtilities.cpp
@@ -4,6 +4,7 @@
 #include "SNiagaraParameterEditor.h"
 #include "NiagaraTypes.h"
 #include "NiagaraEditorStyle.h"
+#include "Engine.h"
 
 #include "SColorPicker.h"
 #include "SNumericEntryBox.h"
diff --git a/Engine/Source/Editor/NiagaraEditor/Private/ViewModels/NiagaraConvertNodeViewModel.h b/Engine/Source/Editor/NiagaraEditor/Private/ViewModels/NiagaraConvertNodeViewModel.h
--- a/Engine/Source/Editor/NiagaraEditor/Private/ViewModels/NiagaraConvertNodeViewModel.h
+++ b/Engine/Source/Editor/NiagaraEditor/Private/ViewModels/NiagaraConvertNodeViewModel.h
@@ -1,6 +1,8 @@
 // Copyright 1998-2016 Epic Games, Inc. All Rights Reserved.
 
 #pragma once
+#include "CoreMinimal.h"
+#include "EdGraph/EdGraphNode.h"
 
 class UNiagaraNodeConvert;
 class FNiagaraConvertPinViewModel;
diff --git a/Engine/Source/Editor/NiagaraEditor/Private/ViewModels/NiagaraConvertPinSocketViewModel.h b/Engine/Source/Editor/NiagaraEditor/Private/ViewModels/NiagaraConvertPinSocketViewModel.h
--- a/Engine/Source/Editor/NiagaraEditor/Private/ViewModels/NiagaraConvertPinSocketViewModel.h
+++ b/Engine/Source/Editor/NiagaraEditor/Private/ViewModels/NiagaraConvertPinSocketViewModel.h
@@ -1,6 +1,9 @@
 // Copyright 1998-2016 Epic Games, Inc. All Rights Reserved.
 
 #pragma once
+#include "SharedPointer.h"
+#include "EdGraph/EdGraphNode.h"
+#include "Visibility.h"
 
 class FNiagaraConvertPinViewModel;
 class FNiagaraConvertNodeViewModel;
diff --git a/Engine/Source/Editor/NiagaraEditor/Private/ViewModels/NiagaraConvertPinViewModel.h b/Engine/Source/Editor/NiagaraEditor/Private/ViewModels/NiagaraConvertPinViewModel.h
--- a/Engine/Source/Editor/NiagaraEditor/Private/ViewModels/NiagaraConvertPinViewModel.h
+++ b/Engine/Source/Editor/NiagaraEditor/Private/ViewModels/NiagaraConvertPinViewModel.h
@@ -1,9 +1,11 @@
 // Copyright 1998-2016 Epic Games, Inc. All Rights Reserved.
 
 #pragma once
+#include "SharedPointer.h"
 
 class FNiagaraConvertNodeViewModel;
 class FNiagaraConvertPinSocketViewModel;
+class UEdGraphPin;
 
 /** A view model for a pin in a convert node. */
 class FNiagaraConvertPinViewModel : public TSharedFromThis<FNiagaraConvertPinViewModel>
diff --git a/Engine/Source/Editor/NiagaraEditor/Private/ViewModels/NiagaraEffectViewModel.cpp b/Engine/Source/Editor/NiagaraEditor/Private/ViewModels/NiagaraEffectViewModel.cpp
--- a/Engine/Source/Editor/NiagaraEditor/Private/ViewModels/NiagaraEffectViewModel.cpp
+++ b/Engine/Source/Editor/NiagaraEditor/Private/ViewModels/NiagaraEffectViewModel.cpp
@@ -16,10 +16,13 @@
 #include "NiagaraEmitterTrackEditor.h"
 #include "NiagaraEffectInstance.h"
 
+#include "Editor.h"
+
 #include "ScopedTransaction.h"
 #include "MovieScene.h"
 #include "ISequencerModule.h"
 #include "EditorSupportDelegates.h"
+#include "ModuleManager.h"
 
 #define LOCTEXT_NAMESPACE "NiagaraEffectViewModel"
 
diff --git a/Engine/Source/Editor/NiagaraEditor/Private/ViewModels/NiagaraEffectViewModel.h b/Engine/Source/Editor/NiagaraEditor/Private/ViewModels/NiagaraEffectViewModel.h
--- a/Engine/Source/Editor/NiagaraEditor/Private/ViewModels/NiagaraEffectViewModel.h
+++ b/Engine/Source/Editor/NiagaraEditor/Private/ViewModels/NiagaraEffectViewModel.h
@@ -5,6 +5,8 @@
 #include "IMovieScenePlayer.h"
 
 #include "NiagaraParameterEditMode.h"
+#include "EditorUndoClient.h"
+#include "GCObject.h"
 
 class UNiagaraEffect;
 class UNiagaraComponent;
@@ -15,6 +17,7 @@ class FNiagaraEmitterHandleViewModel;
 class FNiagaraEffectScriptViewModel;
 class FNiagaraEffectInstance;
 class ISequencer;
+class FAssetData;
 
 /** Defines options for the niagara effect view model */
 struct FNiagaraEffectViewModelOptions
diff --git a/Engine/Source/Editor/NiagaraEditor/Private/ViewModels/NiagaraEmitterHandleViewModel.cpp b/Engine/Source/Editor/NiagaraEditor/Private/ViewModels/NiagaraEmitterHandleViewModel.cpp
--- a/Engine/Source/Editor/NiagaraEditor/Private/ViewModels/NiagaraEmitterHandleViewModel.cpp
+++ b/Engine/Source/Editor/NiagaraEditor/Private/ViewModels/NiagaraEmitterHandleViewModel.cpp
@@ -12,6 +12,7 @@
 #include "NiagaraScriptOutputCollectionViewModel.h"
 
 #include "ScopedTransaction.h"
+#include "AssetEditorManager.h"
 
 #define LOCTEXT_NAMESPACE "EmitterHandleViewModel"
 
diff --git a/Engine/Source/Editor/NiagaraEditor/Private/ViewModels/NiagaraEmitterHandleViewModel.h b/Engine/Source/Editor/NiagaraEditor/Private/ViewModels/NiagaraEmitterHandleViewModel.h
--- a/Engine/Source/Editor/NiagaraEditor/Private/ViewModels/NiagaraEmitterHandleViewModel.h
+++ b/Engine/Source/Editor/NiagaraEditor/Private/ViewModels/NiagaraEmitterHandleViewModel.h
@@ -2,7 +2,11 @@
 
 #pragma once
 
+#include "SharedPointer.h"
+#include "Delegate.h"
 #include "NiagaraParameterEditMode.h"
+#include "SlateEnums.h"
+#include "SlateTypes.h"
 
 class UNiagaraEffect;
 struct FNiagaraEmitterHandle;
diff --git a/Engine/Source/Editor/NiagaraEditor/Private/ViewModels/NiagaraParameterCollectionViewModel.h b/Engine/Source/Editor/NiagaraEditor/Private/ViewModels/NiagaraParameterCollectionViewModel.h
--- a/Engine/Source/Editor/NiagaraEditor/Private/ViewModels/NiagaraParameterCollectionViewModel.h
+++ b/Engine/Source/Editor/NiagaraEditor/Private/ViewModels/NiagaraParameterCollectionViewModel.h
@@ -4,6 +4,7 @@
 
 #include "NiagaraObjectSelection.h"
 #include "NiagaraParameterViewModel.h"
+#include "Visibility.h"
 
 struct FNiagaraTypeDefinition;
 
diff --git a/Engine/Source/Editor/NiagaraEditor/Private/ViewModels/NiagaraParameterViewModel.h b/Engine/Source/Editor/NiagaraEditor/Private/ViewModels/NiagaraParameterViewModel.h
--- a/Engine/Source/Editor/NiagaraEditor/Private/ViewModels/NiagaraParameterViewModel.h
+++ b/Engine/Source/Editor/NiagaraEditor/Private/ViewModels/NiagaraParameterViewModel.h
@@ -3,6 +3,11 @@
 #pragma once
 
 #include "NiagaraParameterEditMode.h"
+#include "DelegateCombinations.h"
+#include "Delegate.h"
+#include "SlateTypes.h"
+#include "StructOnScope.h"
+#include "Text.h"
 
 struct FNiagaraTypeDefinition;
 struct FNiagaraVariable;
diff --git a/Engine/Source/Editor/NiagaraEditor/Private/ViewModels/NiagaraScriptGraphViewModel.cpp b/Engine/Source/Editor/NiagaraEditor/Private/ViewModels/NiagaraScriptGraphViewModel.cpp
--- a/Engine/Source/Editor/NiagaraEditor/Private/ViewModels/NiagaraScriptGraphViewModel.cpp
+++ b/Engine/Source/Editor/NiagaraEditor/Private/ViewModels/NiagaraScriptGraphViewModel.cpp
@@ -10,6 +10,8 @@
 #include "GenericCommands.h"
 #include "ScopedTransaction.h"
 #include "EdGraphUtilities.h"
+#include "Editor.h"
+#include "EditorStyleSet.h"
 
 
 #define LOCTEXT_NAMESPACE "NiagaraScriptGraphViewModel"
diff --git a/Engine/Source/Editor/NiagaraEditor/Private/ViewModels/NiagaraScriptGraphViewModel.h b/Engine/Source/Editor/NiagaraEditor/Private/ViewModels/NiagaraScriptGraphViewModel.h
--- a/Engine/Source/Editor/NiagaraEditor/Private/ViewModels/NiagaraScriptGraphViewModel.h
+++ b/Engine/Source/Editor/NiagaraEditor/Private/ViewModels/NiagaraScriptGraphViewModel.h
@@ -3,9 +3,11 @@
 #pragma once
 
 #include "CoreMinimal.h"
+#include "EdGraph/EdGraphNode.h"
 #include "NiagaraObjectSelection.h"
 #include "EditorUndoClient.h"
 #include "Styling/SlateColor.h"
+#include "UICommandList.h"
 
 class UNiagaraScript;
 class UNiagaraGraph;
diff --git a/Engine/Source/Editor/NiagaraEditor/Private/ViewModels/NiagaraScriptParameterViewModel.cpp b/Engine/Source/Editor/NiagaraEditor/Private/ViewModels/NiagaraScriptParameterViewModel.cpp
--- a/Engine/Source/Editor/NiagaraEditor/Private/ViewModels/NiagaraScriptParameterViewModel.cpp
+++ b/Engine/Source/Editor/NiagaraEditor/Private/ViewModels/NiagaraScriptParameterViewModel.cpp
@@ -5,6 +5,7 @@
 #include "NiagaraTypes.h"
 
 #include "ScopedTransaction.h"
+#include "Editor.h"
 
 #define LOCTEXT_NAMESPACE "ScriptParameterViewModel"
 
diff --git a/Engine/Source/Editor/NiagaraEditor/Private/ViewModels/NiagaraScriptViewModel.cpp b/Engine/Source/Editor/NiagaraEditor/Private/ViewModels/NiagaraScriptViewModel.cpp
--- a/Engine/Source/Editor/NiagaraEditor/Private/ViewModels/NiagaraScriptViewModel.cpp
+++ b/Engine/Source/Editor/NiagaraEditor/Private/ViewModels/NiagaraScriptViewModel.cpp
@@ -12,6 +12,7 @@
 #include "CompilerResultsLog.h"
 #include "GraphEditAction.h"
 #include "Package.h"
+#include "Editor.h"
 
 FNiagaraScriptViewModel::FNiagaraScriptViewModel(UNiagaraScript* InScript, FText DisplayName, ENiagaraParameterEditMode InParameterEditMode)
 	: Script(InScript)
diff --git a/Engine/Source/Editor/NiagaraEditor/Private/ViewModels/NiagaraScriptViewModel.h b/Engine/Source/Editor/NiagaraEditor/Private/ViewModels/NiagaraScriptViewModel.h
--- a/Engine/Source/Editor/NiagaraEditor/Private/ViewModels/NiagaraScriptViewModel.h
+++ b/Engine/Source/Editor/NiagaraEditor/Private/ViewModels/NiagaraScriptViewModel.h
@@ -4,6 +4,8 @@
 
 #include "NiagaraParameterEditMode.h"
 #include "INiagaraCompiler.h"
+#include "EditorUndoClient.h"
+#include "NotifyHook.h"
 
 class UNiagaraScript;
 class INiagaraParameterCollectionViewModel;
diff --git a/Engine/Source/Editor/NiagaraEditor/Private/Widgets/SNiagaraConvertPinSocket.cpp b/Engine/Source/Editor/NiagaraEditor/Private/Widgets/SNiagaraConvertPinSocket.cpp
--- a/Engine/Source/Editor/NiagaraEditor/Private/Widgets/SNiagaraConvertPinSocket.cpp
+++ b/Engine/Source/Editor/NiagaraEditor/Private/Widgets/SNiagaraConvertPinSocket.cpp
@@ -5,7 +5,10 @@
 #include "NiagaraConvertPinViewModel.h"
 #include "NiagaraConvertPinSocketViewModel.h"
 
+#include "SlateApplication.h"
 #include "ScopedTransaction.h"
+#include "WidgetPath.h"
+#include "MultiBoxBuilder.h"
 
 #define LOCTEXT_NAMESPACE "NiagaraConvertPinSocket"
 
diff --git a/Engine/Source/Editor/NiagaraEditor/Private/Widgets/SNiagaraConvertPinSocket.h b/Engine/Source/Editor/NiagaraEditor/Private/Widgets/SNiagaraConvertPinSocket.h
--- a/Engine/Source/Editor/NiagaraEditor/Private/Widgets/SNiagaraConvertPinSocket.h
+++ b/Engine/Source/Editor/NiagaraEditor/Private/Widgets/SNiagaraConvertPinSocket.h
@@ -5,6 +5,7 @@
 #include "DecoratedDragDropOp.h"
 
 class FNiagaraConvertPinSocketViewModel;
+class FMenuBuilder;
 
 /** Contains data for a socket drag and drop operation in the convert node. */
 class FNiagaraConvertDragDropOp : public FDecoratedDragDropOp
diff --git a/Engine/Source/Editor/NiagaraEditor/Private/Widgets/SNiagaraEffectEditor.cpp b/Engine/Source/Editor/NiagaraEditor/Private/Widgets/SNiagaraEffectEditor.cpp
--- a/Engine/Source/Editor/NiagaraEditor/Private/Widgets/SNiagaraEffectEditor.cpp
+++ b/Engine/Source/Editor/NiagaraEditor/Private/Widgets/SNiagaraEffectEditor.cpp
@@ -19,6 +19,7 @@
 #include "Widgets/Layout/SScrollBox.h"
 #include "Widgets/Input/SComboBox.h"
 #include "Widgets/Layout/SExpandableArea.h"
+#include "MultiBoxBuilder.h"
 #include "PropertyEditorModule.h"
 #include "IDetailsView.h"
 
diff --git a/Engine/Source/Editor/NiagaraEditor/Private/Widgets/SNiagaraEffectEditor.h b/Engine/Source/Editor/NiagaraEditor/Private/Widgets/SNiagaraEffectEditor.h
--- a/Engine/Source/Editor/NiagaraEditor/Private/Widgets/SNiagaraEffectEditor.h
+++ b/Engine/Source/Editor/NiagaraEditor/Private/Widgets/SNiagaraEffectEditor.h
@@ -3,10 +3,13 @@
 #pragma once
 
 #include "CoreMinimal.h"
+#include "SCompoundWidget.h"
+#include "DeclarativeSyntaxSupport.h"
 
 class FNiagaraEffectViewModel;
 class FNiagaraEmitterHandleViewModel;
 class SScrollBox;
+class SWidget;
 
 class SNiagaraEffectEditor : public SCompoundWidget
 {
diff --git a/Engine/Source/Editor/NiagaraEditor/Private/Widgets/SNiagaraEffectScript.cpp b/Engine/Source/Editor/NiagaraEditor/Private/Widgets/SNiagaraEffectScript.cpp
--- a/Engine/Source/Editor/NiagaraEditor/Private/Widgets/SNiagaraEffectScript.cpp
+++ b/Engine/Source/Editor/NiagaraEditor/Private/Widgets/SNiagaraEffectScript.cpp
@@ -8,6 +8,8 @@
 #include "SNiagaraParameterCollection.h"
 #include "SNiagaraScriptGraph.h"
 
+#include "SSplitter.h"
+
 void SNiagaraEffectScript::Construct(const FArguments& InArgs, TSharedRef<FNiagaraEffectViewModel> InEffectViewModel)
 {
 	EffectViewModel = InEffectViewModel;
diff --git a/Engine/Source/Editor/NiagaraEditor/Private/Widgets/SNiagaraEffectScript.h b/Engine/Source/Editor/NiagaraEditor/Private/Widgets/SNiagaraEffectScript.h
--- a/Engine/Source/Editor/NiagaraEditor/Private/Widgets/SNiagaraEffectScript.h
+++ b/Engine/Source/Editor/NiagaraEditor/Private/Widgets/SNiagaraEffectScript.h
@@ -1,6 +1,8 @@
 // Copyright 1998-2016 Epic Games, Inc. All Rights Reserved.
 
 #pragma once
+#include "SCompoundWidget.h"
+#include "DeclarativeSyntaxSupport.h"
 
 class FNiagaraEffectViewModel;
 
diff --git a/Engine/Source/Editor/NiagaraEditor/Private/Widgets/SNiagaraEffectTimeline.cpp b/Engine/Source/Editor/NiagaraEditor/Private/Widgets/SNiagaraEffectTimeline.cpp
--- a/Engine/Source/Editor/NiagaraEditor/Private/Widgets/SNiagaraEffectTimeline.cpp
+++ b/Engine/Source/Editor/NiagaraEditor/Private/Widgets/SNiagaraEffectTimeline.cpp
@@ -3,6 +3,7 @@
 #include "SNiagaraEffectTimeline.h"
 
 #include "SCurveEditor.h"
+#include "SOverlay.h"
 
 void SNiagaraTimeline::Construct(const FArguments& InArgs)
 {
diff --git a/Engine/Source/Editor/NiagaraEditor/Private/Widgets/SNiagaraEffectTimeline.h b/Engine/Source/Editor/NiagaraEditor/Private/Widgets/SNiagaraEffectTimeline.h
--- a/Engine/Source/Editor/NiagaraEditor/Private/Widgets/SNiagaraEffectTimeline.h
+++ b/Engine/Source/Editor/NiagaraEditor/Private/Widgets/SNiagaraEffectTimeline.h
@@ -1,11 +1,14 @@
 // Copyright 1998-2016 Epic Games, Inc. All Rights Reserved.
 
 #pragma once
+#include "SCompoundWidget.h"
+#include "DeclarativeSyntaxSupport.h"
 
 class SCurveEditor;
 struct FNiagaraSimulation;
 class FNiagaraEffectInstance;
 class UNiagaraEffect;
+class UCurveBase;
 
 class SNiagaraTimeline : public SCompoundWidget
 {
diff --git a/Engine/Source/Editor/NiagaraEditor/Private/Widgets/SNiagaraEmitterEditor.h b/Engine/Source/Editor/NiagaraEditor/Private/Widgets/SNiagaraEmitterEditor.h
--- a/Engine/Source/Editor/NiagaraEditor/Private/Widgets/SNiagaraEmitterEditor.h
+++ b/Engine/Source/Editor/NiagaraEditor/Private/Widgets/SNiagaraEmitterEditor.h
@@ -1,6 +1,8 @@
 // Copyright 1998-2016 Epic Games, Inc. All Rights Reserved.
 
 #pragma once
+#include "SCompoundWidget.h"
+#include "DeclarativeSyntaxSupport.h"
 
 class FNiagaraEmitterViewModel;
 class IDetailsView;
diff --git a/Engine/Source/Editor/NiagaraEditor/Private/Widgets/SNiagaraEmitterHeader.cpp b/Engine/Source/Editor/NiagaraEditor/Private/Widgets/SNiagaraEmitterHeader.cpp
--- a/Engine/Source/Editor/NiagaraEditor/Private/Widgets/SNiagaraEmitterHeader.cpp
+++ b/Engine/Source/Editor/NiagaraEditor/Private/Widgets/SNiagaraEmitterHeader.cpp
@@ -4,8 +4,10 @@
 #include "NiagaraEmitterHandleViewModel.h"
 #include "NiagaraEmitterViewModel.h"
 #include "NiagaraEditorStyle.h"
-
 #include "SInlineEditableTextBlock.h"
+#include "SBoxPanel.h"
+#include "SCheckBox.h"
+#include "STextBlock.h"
 
 #define LOCTEXT_NAMESPACE "NiagaraEmitterHeader"
 
diff --git a/Engine/Source/Editor/NiagaraEditor/Private/Widgets/SNiagaraEmitterHeader.h b/Engine/Source/Editor/NiagaraEditor/Private/Widgets/SNiagaraEmitterHeader.h
--- a/Engine/Source/Editor/NiagaraEditor/Private/Widgets/SNiagaraEmitterHeader.h
+++ b/Engine/Source/Editor/NiagaraEditor/Private/Widgets/SNiagaraEmitterHeader.h
@@ -1,6 +1,9 @@
 // Copyright 1998-2016 Epic Games, Inc. All Rights Reserved.
 
 #pragma once
+#include "SharedPointer.h"
+#include "SCompoundWidget.h"
+#include "DeclarativeSyntaxSupport.h"
 
 class FNiagaraEmitterHandleViewModel;
 
diff --git a/Engine/Source/Editor/NiagaraEditor/Private/Widgets/SNiagaraGraphNodeConvert.cpp b/Engine/Source/Editor/NiagaraEditor/Private/Widgets/SNiagaraGraphNodeConvert.cpp
--- a/Engine/Source/Editor/NiagaraEditor/Private/Widgets/SNiagaraGraphNodeConvert.cpp
+++ b/Engine/Source/Editor/NiagaraEditor/Private/Widgets/SNiagaraGraphNodeConvert.cpp
@@ -8,6 +8,7 @@
 #include "SNiagaraConvertPinSocket.h"
 
 #include "GraphEditorSettings.h"
+#include "DrawElements.h"
 #include "SGraphPin.h"
 
 void SNiagaraGraphNodeConvert::Construct(const FArguments& InArgs, UEdGraphNode* InGraphNode)
diff --git a/Engine/Source/Editor/NiagaraEditor/Private/Widgets/SNiagaraGraphPinAdd.cpp b/Engine/Source/Editor/NiagaraEditor/Private/Widgets/SNiagaraGraphPinAdd.cpp
--- a/Engine/Source/Editor/NiagaraEditor/Private/Widgets/SNiagaraGraphPinAdd.cpp
+++ b/Engine/Source/Editor/NiagaraEditor/Private/Widgets/SNiagaraGraphPinAdd.cpp
@@ -4,6 +4,10 @@
 #include "NiagaraNodeWithDynamicPins.h"
 
 #include "ScopedTransaction.h"
+#include "SComboButton.h"
+#include "SBoxPanel.h"
+#include "SImage.h"
+#include "MultiBoxBuilder.h"
 
 #define LOCTEXT_NAMESPACE "NiagaraGraphPinAdd"
 
diff --git a/Engine/Source/Editor/NiagaraEditor/Private/Widgets/SNiagaraParameterCollection.cpp b/Engine/Source/Editor/NiagaraEditor/Private/Widgets/SNiagaraParameterCollection.cpp
--- a/Engine/Source/Editor/NiagaraEditor/Private/Widgets/SNiagaraParameterCollection.cpp
+++ b/Engine/Source/Editor/NiagaraEditor/Private/Widgets/SNiagaraParameterCollection.cpp
@@ -9,15 +9,21 @@
 #include "SNiagaraParameterEditor.h"
 #include "NiagaraEditorStyle.h"
 
+#include "SButton.h"
 #include "SExpandableArea.h"
 #include "SInlineEditableTextBlock.h"
 #include "SComboBox.h"
+#include "SImage.h"
+#include "SSplitter.h"
+#include "MultiBoxBuilder.h"
+
 #include "Modules/ModuleManager.h"
 #include "IStructureDetailsView.h"
 #include "GenericCommands.h"
 #include "Editor/PropertyEditor/Public/PropertyEditorModule.h"
 #include "IDetailsView.h"
 
+
 #define LOCTEXT_NAMESPACE "NiagaraParameterCollectionEditor"
 
 class SSimpleExpander : public SCompoundWidget
diff --git a/Engine/Source/Editor/NiagaraEditor/Private/Widgets/SNiagaraParameterCollection.h b/Engine/Source/Editor/NiagaraEditor/Private/Widgets/SNiagaraParameterCollection.h
--- a/Engine/Source/Editor/NiagaraEditor/Private/Widgets/SNiagaraParameterCollection.h
+++ b/Engine/Source/Editor/NiagaraEditor/Private/Widgets/SNiagaraParameterCollection.h
@@ -2,6 +2,11 @@
 
 #pragma once
 
+#include "SCompoundWidget.h"
+#include "NotifyHook.h"
+#include "SListView.h"
+
+
 class INiagaraParameterCollectionViewModel;
 class INiagaraParameterViewModel;
 class IStructureDetailsView;
@@ -9,6 +14,13 @@ class SNiagaraParameterEditor;
 struct FNiagaraTypeDefinition;
 struct FNiagaraVariable;
 
+class FUICommandList;
+class ITableRow;
+class STableViewBase;
+class SExpandableArea;
+class SBox;
+class SComboButton;
+
 /** A widget for editing parameter collections. */
 class SNiagaraParameterCollection : public SCompoundWidget, public FNotifyHook
 {
diff --git a/Engine/Source/Editor/NiagaraEditor/Private/Widgets/SNiagaraScriptGraph.cpp b/Engine/Source/Editor/NiagaraEditor/Private/Widgets/SNiagaraScriptGraph.cpp
--- a/Engine/Source/Editor/NiagaraEditor/Private/Widgets/SNiagaraScriptGraph.cpp
+++ b/Engine/Source/Editor/NiagaraEditor/Private/Widgets/SNiagaraScriptGraph.cpp
@@ -6,7 +6,15 @@
 #include "NiagaraGraph.h"
 #include "NiagaraNode.h"
 
+#include "AssetEditorManager.h"
 #include "GraphEditor.h"
+#include "EditorStyleSet.h"
+#include "SBorder.h"
+#include "SBoxPanel.h"
+#include "TextLayout.h"
+#include "SErrorText.h"
+#include "STextBlock.h"
+
 
 #define LOCTEXT_NAMESPACE "NiagaraScriptGraph"
 
diff --git a/Engine/Source/Editor/NiagaraEditor/Private/Widgets/SNiagaraScriptGraph.h b/Engine/Source/Editor/NiagaraEditor/Private/Widgets/SNiagaraScriptGraph.h
--- a/Engine/Source/Editor/NiagaraEditor/Private/Widgets/SNiagaraScriptGraph.h
+++ b/Engine/Source/Editor/NiagaraEditor/Private/Widgets/SNiagaraScriptGraph.h
@@ -3,6 +3,9 @@
 #pragma once
 
 #include "SCompoundWidget.h"
+#include "DeclarativeSyntaxSupport.h"
+#include "GraphEditor.h"
+#include "EdGraph/EdGraphNode.h"
 
 class FNiagaraScriptGraphViewModel;
 
diff --git a/Engine/Source/Editor/NiagaraEditor/Private/Widgets/SNiagaraSelectedObjectsDetails.h b/Engine/Source/Editor/NiagaraEditor/Private/Widgets/SNiagaraSelectedObjectsDetails.h
--- a/Engine/Source/Editor/NiagaraEditor/Private/Widgets/SNiagaraSelectedObjectsDetails.h
+++ b/Engine/Source/Editor/NiagaraEditor/Private/Widgets/SNiagaraSelectedObjectsDetails.h
@@ -1,7 +1,10 @@
 // Copyright 1998-2016 Epic Games, Inc. All Rights Reserved.
 
 #pragma once
+
 #include "SCompoundWidget.h"
+#include "Delegate.h"
+#include "DeclarativeSyntaxSupport.h"
 
 class FNiagaraObjectSelection;
 class IDetailsView;
diff --git a/Engine/Source/Editor/NiagaraEditor/Public/INiagaraEditorTypeUtilities.h b/Engine/Source/Editor/NiagaraEditor/Public/INiagaraEditorTypeUtilities.h
--- a/Engine/Source/Editor/NiagaraEditor/Public/INiagaraEditorTypeUtilities.h
+++ b/Engine/Source/Editor/NiagaraEditor/Public/INiagaraEditorTypeUtilities.h
@@ -1,7 +1,10 @@
 // Copyright 1998-2016 Epic Games, Inc. All Rights Reserved.
 
 #pragma once
+#include "SharedPointer.h"
+#include "Delegate.h"
 
+class FStructOnScope;
 class SNiagaraParameterEditor;
 
 class INiagaraEditorTypeUtilities
diff --git a/Engine/Source/Editor/NiagaraEditor/Public/NiagaraEditorStyle.h b/Engine/Source/Editor/NiagaraEditor/Public/NiagaraEditorStyle.h
--- a/Engine/Source/Editor/NiagaraEditor/Public/NiagaraEditorStyle.h
+++ b/Engine/Source/Editor/NiagaraEditor/Public/NiagaraEditorStyle.h
@@ -1,6 +1,7 @@
 // Copyright 1998-2016 Epic Games, Inc. All Rights Reserved.
 
 #pragma once
+#include "ISlateStyle.h"
 
 /** Manages the style which provides resources for niagara editor widgets. */
 class FNiagaraEditorStyle
diff --git a/Engine/Source/Editor/NiagaraEditor/Public/NiagaraEmitterFactoryNew.h b/Engine/Source/Editor/NiagaraEditor/Public/NiagaraEmitterFactoryNew.h
--- a/Engine/Source/Editor/NiagaraEditor/Public/NiagaraEmitterFactoryNew.h
+++ b/Engine/Source/Editor/NiagaraEditor/Public/NiagaraEmitterFactoryNew.h
@@ -2,6 +2,8 @@
 
 #pragma once
 #include "NiagaraSettings.h"
+#include "Factories/Factory.h"
+
 #include "NiagaraEmitterFactoryNew.generated.h"
 
 
diff --git a/Engine/Source/Editor/NiagaraEditor/Public/SNiagaraParameterEditor.h b/Engine/Source/Editor/NiagaraEditor/Public/SNiagaraParameterEditor.h
--- a/Engine/Source/Editor/NiagaraEditor/Public/SNiagaraParameterEditor.h
+++ b/Engine/Source/Editor/NiagaraEditor/Public/SNiagaraParameterEditor.h
@@ -1,6 +1,8 @@
 // Copyright 1998-2016 Epic Games, Inc. All Rights Reserved.
 
 #pragma once
+#include "SCompoundWidget.h"
+#include "StructOnScope.h"
 
 /** Base class for inline parameter editors. These editors are expected to maintain an internal value which
 	is populated from a parameter struct. */
diff --git a/Engine/Source/Runtime/Niagara/Classes/NiagaraCollision.h b/Engine/Source/Runtime/Niagara/Classes/NiagaraCollision.h
--- a/Engine/Source/Runtime/Niagara/Classes/NiagaraCollision.h
+++ b/Engine/Source/Runtime/Niagara/Classes/NiagaraCollision.h
@@ -1,7 +1,7 @@
 #pragma once
 
 #include "NiagaraEvents.h"
-
+#include "WorldCollision.h"
 
 UENUM()
 enum class ENiagaraCollisionMode : uint8
diff --git a/Engine/Source/Runtime/Niagara/Classes/NiagaraEmitterHandle.h b/Engine/Source/Runtime/Niagara/Classes/NiagaraEmitterHandle.h
--- a/Engine/Source/Runtime/Niagara/Classes/NiagaraEmitterHandle.h
+++ b/Engine/Source/Runtime/Niagara/Classes/NiagaraEmitterHandle.h
@@ -2,6 +2,10 @@
 
 #pragma once
 
+#include "CoreMinimal.h"
+#include "UObject/ObjectMacros.h"
+#include "UObject/Object.h"
+#include "Guid.h"
 #include "NiagaraEmitterHandle.generated.h"
 
 class UNiagaraEffect;
diff --git a/Engine/Source/Runtime/Niagara/Private/NiagaraModule.cpp b/Engine/Source/Runtime/Niagara/Private/NiagaraModule.cpp
--- a/Engine/Source/Runtime/Niagara/Private/NiagaraModule.cpp
+++ b/Engine/Source/Runtime/Niagara/Private/NiagaraModule.cpp
@@ -4,6 +4,8 @@
 #include "Modules/ModuleManager.h"
 #include "NiagaraTypes.h"
 #include "NiagaraEvents.h"
+#include "Class.h"
+#include "Package.h"
 
 IMPLEMENT_MODULE(INiagaraModule, Niagara);
 
diff --git a/Engine/Source/Runtime/Niagara/Private/NiagaraScript.cpp b/Engine/Source/Runtime/Niagara/Private/NiagaraScript.cpp
--- a/Engine/Source/Runtime/Niagara/Private/NiagaraScript.cpp
+++ b/Engine/Source/Runtime/Niagara/Private/NiagaraScript.cpp
@@ -3,6 +3,7 @@
 #include "NiagaraScript.h"
 #include "NiagaraScriptSourceBase.h"
 #include "NiagaraComponent.h"
+#include "Package.h"
 
 
 
diff --git a/Engine/Source/Runtime/Niagara/Public/NiagaraEffectRendererProperties.h b/Engine/Source/Runtime/Niagara/Public/NiagaraEffectRendererProperties.h
--- a/Engine/Source/Runtime/Niagara/Public/NiagaraEffectRendererProperties.h
+++ b/Engine/Source/Runtime/Niagara/Public/NiagaraEffectRendererProperties.h
@@ -5,6 +5,7 @@
 #include "CoreMinimal.h"
 #include "UObject/ObjectMacros.h"
 #include "UObject/Object.h"
+#include "RHIDefinitions.h"
 #include "NiagaraEffectRendererProperties.generated.h"
 
 /**
@@ -15,6 +16,7 @@
 */
 
 class NiagaraEffectRenderer;
+class UMaterial;
 
 UCLASS(ABSTRACT)
 class NIAGARA_API UNiagaraEffectRendererProperties : public UObject
