Include less intrusive engine header for GEngine access.

diff --git a/Engine/Source/Editor/NiagaraEditor/Private/TypeEditorUtilities/NiagaraColorTypeEditorUtilities.cpp b/Engine/Source/Editor/NiagaraEditor/Private/TypeEditorUtilities/NiagaraColorTypeEditorUtilities.cpp
--- a/Engine/Source/Editor/NiagaraEditor/Private/TypeEditorUtilities/NiagaraColorTypeEditorUtilities.cpp
+++ b/Engine/Source/Editor/NiagaraEditor/Private/TypeEditorUtilities/NiagaraColorTypeEditorUtilities.cpp
@@ -4,7 +4,7 @@
 #include "SNiagaraParameterEditor.h"
 #include "NiagaraTypes.h"
 #include "NiagaraEditorStyle.h"
-#include "Engine.h"
+#include "Classes/Engine/Engine.h"
 
 #include "SColorPicker.h"
 #include "SNumericEntryBox.h"
