Fixing some localization key conflicts that crept in from stream merges

diff --git a/Engine/Source/Editor/NiagaraEditor/Private/EdGraphSchema_Niagara.cpp b/Engine/Source/Editor/NiagaraEditor/Private/EdGraphSchema_Niagara.cpp
--- a/Engine/Source/Editor/NiagaraEditor/Private/EdGraphSchema_Niagara.cpp
+++ b/Engine/Source/Editor/NiagaraEditor/Private/EdGraphSchema_Niagara.cpp
@@ -555,7 +555,7 @@ void UEdGraphSchema_Niagara::GetGraphContextActions(FGraphContextMenuBuilder& Co
 		}
 	}
 
-	const FText MenuCat = LOCTEXT("NiagaraEventMenuCat", "Logic");
+	const FText MenuCat = LOCTEXT("NiagaraLogicMenuCat", "Logic");
 	{
 		const FText MenuDesc = LOCTEXT("If", "If");
 
