Copy fix for Mac/Linux editor compilation from DevRendering

diff --git a/Engine/Source/Editor/NiagaraEditor/Private/ViewModels/NiagaraEffectViewModel.cpp b/Engine/Source/Editor/NiagaraEditor/Private/ViewModels/NiagaraEffectViewModel.cpp
--- a/Engine/Source/Editor/NiagaraEditor/Private/ViewModels/NiagaraEffectViewModel.cpp
+++ b/Engine/Source/Editor/NiagaraEditor/Private/ViewModels/NiagaraEffectViewModel.cpp
@@ -26,7 +26,7 @@
 
 #define LOCTEXT_NAMESPACE "NiagaraEffectViewModel"
 
-TMap<UNiagaraEffect*, TArray<FNiagaraEffectViewModel*>> TNiagaraViewModelManager<UNiagaraEffect, FNiagaraEffectViewModel>::ObjectsToViewModels;
+template<> TMap<UNiagaraEffect*, TArray<FNiagaraEffectViewModel*>> TNiagaraViewModelManager<UNiagaraEffect, FNiagaraEffectViewModel>::ObjectsToViewModels;
 
 FNiagaraEffectViewModel::FNiagaraEffectViewModel(UNiagaraEffect& InEffect, FNiagaraEffectViewModelOptions InOptions)
 	: Effect(InEffect)
diff --git a/Engine/Source/Editor/NiagaraEditor/Private/ViewModels/NiagaraEmitterViewModel.cpp b/Engine/Source/Editor/NiagaraEditor/Private/ViewModels/NiagaraEmitterViewModel.cpp
--- a/Engine/Source/Editor/NiagaraEditor/Private/ViewModels/NiagaraEmitterViewModel.cpp
+++ b/Engine/Source/Editor/NiagaraEditor/Private/ViewModels/NiagaraEmitterViewModel.cpp
@@ -16,7 +16,7 @@
 
 #define LOCTEXT_NAMESPACE "EmitterEditorViewModel"
 
-TMap<UNiagaraEmitterProperties*, TArray<FNiagaraEmitterViewModel*>> TNiagaraViewModelManager<UNiagaraEmitterProperties, FNiagaraEmitterViewModel>::ObjectsToViewModels;
+template<> TMap<UNiagaraEmitterProperties*, TArray<FNiagaraEmitterViewModel*>> TNiagaraViewModelManager<UNiagaraEmitterProperties, FNiagaraEmitterViewModel>::ObjectsToViewModels;
 
 const FText FNiagaraEmitterViewModel::StatsFormat = NSLOCTEXT("NiagaraEmitterViewModel", "StatsFormat", "{0} Particles | {1} ms | {2} MB");
 const float Megabyte = 1024.0f * 1024.0f;
diff --git a/Engine/Source/Editor/NiagaraEditor/Private/ViewModels/NiagaraScriptViewModel.cpp b/Engine/Source/Editor/NiagaraEditor/Private/ViewModels/NiagaraScriptViewModel.cpp
--- a/Engine/Source/Editor/NiagaraEditor/Private/ViewModels/NiagaraScriptViewModel.cpp
+++ b/Engine/Source/Editor/NiagaraEditor/Private/ViewModels/NiagaraScriptViewModel.cpp
@@ -15,7 +15,7 @@
 #include "Editor.h"
 #include "TNiagaraViewModelManager.h"
 
-TMap<UNiagaraScript*, TArray<FNiagaraScriptViewModel*>> TNiagaraViewModelManager<UNiagaraScript, FNiagaraScriptViewModel>::ObjectsToViewModels;
+template<> TMap<UNiagaraScript*, TArray<FNiagaraScriptViewModel*>> TNiagaraViewModelManager<UNiagaraScript, FNiagaraScriptViewModel>::ObjectsToViewModels;
 
 FNiagaraScriptViewModel::FNiagaraScriptViewModel(UNiagaraScript* InScript, FText DisplayName, ENiagaraParameterEditMode InParameterEditMode)
 	: Script(InScript)
