Fixing static analysis warnings.

diff --git a/Engine/Source/Editor/NiagaraEditor/Private/NiagaraCompiler.cpp b/Engine/Source/Editor/NiagaraEditor/Private/NiagaraCompiler.cpp
--- a/Engine/Source/Editor/NiagaraEditor/Private/NiagaraCompiler.cpp
+++ b/Engine/Source/Editor/NiagaraEditor/Private/NiagaraCompiler.cpp
@@ -769,6 +769,7 @@ UNiagaraGraph* FHlslNiagaraCompiler::CloneGraphAndPrepareForCompilation(UNiagara
 	if (Source == nullptr)
 	{
 		Error(LOCTEXT("NoSourceErrorMessage", "Script has no source."), nullptr, nullptr);
+		return nullptr;
 	}
 
 	if (bClearErrors)
diff --git a/Engine/Source/Editor/NiagaraEditor/Private/NiagaraNodeInput.cpp b/Engine/Source/Editor/NiagaraEditor/Private/NiagaraNodeInput.cpp
--- a/Engine/Source/Editor/NiagaraEditor/Private/NiagaraNodeInput.cpp
+++ b/Engine/Source/Editor/NiagaraEditor/Private/NiagaraNodeInput.cpp
@@ -180,7 +180,7 @@ bool UNiagaraNodeInput::VerifyNodeRenameTextCommit(const FText& NewText, UNiagar
 
 			for (UNiagaraNodeInput* Node : InputNodes)
 			{
-				if (Node == InputNodeBeingChanged || Node->Usage != InputNodeBeingChanged->Usage)
+				if (Node == nullptr || Node == InputNodeBeingChanged || Node->Usage != InputNodeBeingChanged->Usage)
 				{
 					continue;
 				}
@@ -231,7 +231,7 @@ void UNiagaraNodeInput::OnRenameNode(const FString& NewName)
 	AffectedNodes.Add(this);
 	for (UNiagaraNodeInput* Node : InputNodes)
 	{
-		if (Node == this)
+		if (Node == this || Node == nullptr)
 		{
 			continue;
 		}
@@ -284,6 +284,11 @@ bool UNiagaraNodeInput::ReferencesSameInput(UNiagaraNodeInput* Other) const
 		return true;
 	}
 
+	if (Other == nullptr)
+	{
+		return false;
+	}
+
 	if (Usage == Other->Usage)
 	{
 		if (Usage == ENiagaraInputNodeUsage::Parameter && Input.GetId() == Other->Input.GetId())
@@ -317,6 +322,11 @@ void UNiagaraNodeInput::AutowireNewNode(UEdGraphPin* FromPin)
 			int32 HighestSortPriority = -1; // Set to -1 initially, so that in the event of no nodes, we still get zero.
 			for (UNiagaraNodeInput* InputNode : InputNodes)
 			{
+				if (InputNode == nullptr)
+				{
+					continue;
+				}
+
 				if (InputNode != this && InputNode->Usage == ENiagaraInputNodeUsage::Parameter)
 				{
 					if (this->ReferencesSameInput(InputNode))
diff --git a/Engine/Source/Editor/NiagaraEditor/Private/ViewModels/NiagaraScriptGraphViewModel.cpp b/Engine/Source/Editor/NiagaraEditor/Private/ViewModels/NiagaraScriptGraphViewModel.cpp
--- a/Engine/Source/Editor/NiagaraEditor/Private/ViewModels/NiagaraScriptGraphViewModel.cpp
+++ b/Engine/Source/Editor/NiagaraEditor/Private/ViewModels/NiagaraScriptGraphViewModel.cpp
@@ -336,11 +336,15 @@ void FixUpPastedInputNodes(UEdGraph* Graph, TSet<UEdGraphNode*> PastedNodes)
 			}
 		}
 
-		if (MatchingInputByNameAndType != nullptr)
+		if (MatchingInputByNameAndType != nullptr && MatchingNode != nullptr)
 		{
 			// Update the id and value on the matching pasted nodes.
 			for (UNiagaraNodeInput* PastedNodeForInput : PastedNodesForInput)
 			{
+				if (nullptr == PastedNodeForInput)
+				{
+					continue;
+				}
 				PastedNodeForInput->CallSortPriority = MatchingNode->CallSortPriority;
 				PastedNodeForInput->ExposureOptions = MatchingNode->ExposureOptions;
 				PastedNodeForInput->Input.SetId(MatchingInputByNameAndType->GetId());
diff --git a/Engine/Source/Runtime/Niagara/Private/NiagaraEffectRendererSprites.cpp b/Engine/Source/Runtime/Niagara/Private/NiagaraEffectRendererSprites.cpp
--- a/Engine/Source/Runtime/Niagara/Private/NiagaraEffectRendererSprites.cpp
+++ b/Engine/Source/Runtime/Niagara/Private/NiagaraEffectRendererSprites.cpp
@@ -94,7 +94,7 @@ void NiagaraEffectRendererSprites::GetDynamicMeshElements(const TArray<const FSc
 	//check(DynamicDataRender)
 
 	FNiagaraDynamicDataSprites *DynamicDataSprites = static_cast<FNiagaraDynamicDataSprites*>(DynamicDataRender);
-	if (!DynamicDataSprites || DynamicDataSprites->VertexData.Num() == 0)
+	if (!DynamicDataSprites || DynamicDataSprites->VertexData.Num() == 0 || nullptr == Properties)
 	{
 		return;
 	}
